DatasmithDeltaGenTranslator: Fix spurious deoptimization flags

diff --git a/Engine/Plugins/Enterprise/DatasmithFBXImporter/Source/DatasmithDeltaGenTranslator/DatasmithDeltaGenTranslator.Build.cs b/Engine/Plugins/Enterprise/DatasmithFBXImporter/Source/DatasmithDeltaGenTranslator/DatasmithDeltaGenTranslator.Build.cs
--- a/Engine/Plugins/Enterprise/DatasmithFBXImporter/Source/DatasmithDeltaGenTranslator/DatasmithDeltaGenTranslator.Build.cs
+++ b/Engine/Plugins/Enterprise/DatasmithFBXImporter/Source/DatasmithDeltaGenTranslator/DatasmithDeltaGenTranslator.Build.cs
@@ -10,10 +10,6 @@ namespace UnrealBuildTool.Rules
 	{
 		public DatasmithDeltaGenTranslator(ReadOnlyTargetRules Target) : base(Target)
 		{
-			OptimizeCode = CodeOptimization.Never;
-			bUseUnity = false;
-			PCHUsage = PCHUsageMode.NoPCHs;
-
 			PrivateDependencyModuleNames.AddRange(
 				new string[]
 				{
