Fix infinite loop causing memory exhaustion when importing fbx with complex variant sets

From dev-enterprise[REDACTED_HANDLE]

diff --git a/Engine/Plugins/Enterprise/DatasmithFBXImporter/Source/DatasmithDeltaGenTranslator/Private/DatasmithDeltaGenVariantConverter.cpp b/Engine/Plugins/Enterprise/DatasmithFBXImporter/Source/DatasmithDeltaGenTranslator/Private/DatasmithDeltaGenVariantConverter.cpp
--- a/Engine/Plugins/Enterprise/DatasmithFBXImporter/Source/DatasmithDeltaGenTranslator/Private/DatasmithDeltaGenVariantConverter.cpp
+++ b/Engine/Plugins/Enterprise/DatasmithFBXImporter/Source/DatasmithDeltaGenTranslator/Private/DatasmithDeltaGenVariantConverter.cpp
@@ -32,6 +32,12 @@ namespace FDeltaGenVariantConverterImpl
 			// It doesn't matter if there are multiple captures to the same property, as only one UPropertyCapture is ever created
 			if (TSharedPtr<IDatasmithActorBindingElement>* FoundBindingInB = BindingsInB.Find(Actor))
 			{
+				// Don't merge a binding into itself
+				if (BindingFromA == *FoundBindingInB)
+				{
+					continue;
+				}
+
 				for (int32 PropIndex = 0; PropIndex < BindingFromA->GetPropertyCapturesCount(); ++PropIndex)
 				{
 					TSharedPtr<IDatasmithBasePropertyCaptureElement> PropertyCapture = BindingFromA->GetPropertyCapture(PropIndex);
