Remove mesh merging from importer as it causes huge performance problems and fonctionality overlaps with more generic engine post import actor merging
   - 6m28s [at] 91.48GB RAM -> 2m42s [at] 38.87GB RAM when importing DATASET-0023a
   - 1m44s [at] 18.45GB RAM -> 0m57s [at] 13.85GB RAM when importing DATASET-0025a

From dev-enterprise[REDACTED_HANDLE]

diff --git a/Engine/Plugins/Enterprise/DatasmithFBXImporter/Source/DatasmithDeltaGenTranslator/Private/DatasmithDeltaGenImportOptions.cpp b/Engine/Plugins/Enterprise/DatasmithFBXImporter/Source/DatasmithDeltaGenTranslator/Private/DatasmithDeltaGenImportOptions.cpp
--- a/Engine/Plugins/Enterprise/DatasmithFBXImporter/Source/DatasmithDeltaGenTranslator/Private/DatasmithDeltaGenImportOptions.cpp
+++ b/Engine/Plugins/Enterprise/DatasmithFBXImporter/Source/DatasmithDeltaGenTranslator/Private/DatasmithDeltaGenImportOptions.cpp
@@ -13,7 +13,6 @@
 
 UDatasmithDeltaGenImportOptions::UDatasmithDeltaGenImportOptions(const FObjectInitializer& ObjectInitializer)
 	: Super(ObjectInitializer)
-	, bMergeNodes(true)
 	, bOptimizeDuplicatedNodes(true)
 	, bRemoveInvisibleNodes(true)
 	, bSimplifyNodeHierarchy(true)
@@ -166,7 +165,6 @@ void UDatasmithDeltaGenImportOptions::FromSceneImportData(UDatasmithFBXSceneImpo
 
 	if (UDatasmithDeltaGenSceneImportData* DGImportData = Cast<UDatasmithDeltaGenSceneImportData>(InImportData))
 	{
-		bMergeNodes					= DGImportData->bMergeNodes;
 		bOptimizeDuplicatedNodes	= DGImportData->bOptimizeDuplicatedNodes;
 		bRemoveInvisibleNodes		= DGImportData->bRemoveInvisibleNodes;
 		bSimplifyNodeHierarchy		= DGImportData->bSimplifyNodeHierarchy;
@@ -185,7 +183,6 @@ void UDatasmithDeltaGenImportOptions::ToSceneImportData(UDatasmithFBXSceneImport
 
 	if (UDatasmithDeltaGenSceneImportData* DGImportData = Cast<UDatasmithDeltaGenSceneImportData>(OutImportData))
 	{
-		DGImportData->bMergeNodes				= bMergeNodes;
 		DGImportData->bOptimizeDuplicatedNodes	= bOptimizeDuplicatedNodes;
 		DGImportData->bRemoveInvisibleNodes		= bRemoveInvisibleNodes;
 		DGImportData->bSimplifyNodeHierarchy	= bSimplifyNodeHierarchy;
diff --git a/Engine/Plugins/Enterprise/DatasmithFBXImporter/Source/DatasmithDeltaGenTranslator/Private/DatasmithDeltaGenImportOptions.h b/Engine/Plugins/Enterprise/DatasmithFBXImporter/Source/DatasmithDeltaGenTranslator/Private/DatasmithDeltaGenImportOptions.h
--- a/Engine/Plugins/Enterprise/DatasmithFBXImporter/Source/DatasmithDeltaGenTranslator/Private/DatasmithDeltaGenImportOptions.h
+++ b/Engine/Plugins/Enterprise/DatasmithFBXImporter/Source/DatasmithDeltaGenTranslator/Private/DatasmithDeltaGenImportOptions.h
@@ -28,9 +28,6 @@ class UDatasmithDeltaGenImportOptions : public UDatasmithFBXImportOptions
 	GENERATED_UCLASS_BODY()
 
 public:
-	UPROPERTY(config, EditAnywhere, BlueprintReadWrite, Category=Processing, meta=(ToolTip="Merge nodes that are not used in variants/animations"))
-	bool bMergeNodes;
-
 	UPROPERTY(config, EditAnywhere, BlueprintReadWrite, Category = Processing, meta = (ToolTip = "Finds duplicated scene nodes which could be replaced with a single instance"))
 	bool bOptimizeDuplicatedNodes;
 
diff --git a/Engine/Plugins/Enterprise/DatasmithFBXImporter/Source/DatasmithDeltaGenTranslator/Private/DatasmithDeltaGenImporter.cpp b/Engine/Plugins/Enterprise/DatasmithFBXImporter/Source/DatasmithDeltaGenTranslator/Private/DatasmithDeltaGenImporter.cpp
--- a/Engine/Plugins/Enterprise/DatasmithFBXImporter/Source/DatasmithDeltaGenTranslator/Private/DatasmithDeltaGenImporter.cpp
+++ b/Engine/Plugins/Enterprise/DatasmithFBXImporter/Source/DatasmithDeltaGenTranslator/Private/DatasmithDeltaGenImporter.cpp
@@ -356,14 +356,6 @@ bool FDatasmithDeltaGenImporter::ProcessScene()
 		Processor.OptimizeDuplicatedNodes();
 	}
 
-	if (ImportOptions->bMergeNodes)
-	{
-		Processor.MergeSceneNodes();
-
-		//Call "RemoveEmptyNodes" again after merging, because all nodes which previously had meshes are empty now.
-		Processor.RemoveEmptyNodes();
-	}
-
 	Processor.FixMeshNames();
 
 	return true;
diff --git a/Engine/Plugins/Enterprise/DatasmithFBXImporter/Source/DatasmithFBXTranslator/Private/DatasmithFBXSceneProcessor.cpp b/Engine/Plugins/Enterprise/DatasmithFBXImporter/Source/DatasmithFBXTranslator/Private/DatasmithFBXSceneProcessor.cpp
--- a/Engine/Plugins/Enterprise/DatasmithFBXImporter/Source/DatasmithFBXTranslator/Private/DatasmithFBXSceneProcessor.cpp
+++ b/Engine/Plugins/Enterprise/DatasmithFBXImporter/Source/DatasmithFBXTranslator/Private/DatasmithFBXSceneProcessor.cpp
@@ -254,459 +254,6 @@ void FDatasmithFBXSceneProcessor::FindPersistentNodes()
 	Helper.Recurse(Scene->RootNode);
 }
 
-struct FNodeMergeHelper
-{
-	void MergeScene(FDatasmithFBXScene* Scene)
-	{
-		//todo: output some stats
-		// Merge all root's children which are marked for merging
-		RecursiveMergeScene(Scene->RootNode);
-		// Merge all root's children which aren't marked (this works like if we had RootNode also marked)
-		MergeNodesWithParent(Scene->RootNode);
-	}
-
-protected:
-	/** Traverse the scene tree and execute merging for every marked node */
-	void RecursiveMergeScene(TSharedPtr<FDatasmithFBXSceneNode>& Node)
-	{
-		for (int32 NodeIndex = 0; NodeIndex < Node->Children.Num(); NodeIndex++)
-		{
-			TSharedPtr<FDatasmithFBXSceneNode>& Child = Node->Children[NodeIndex];
-
-			// Execute merging operation for every marked node.
-			if (Child->bShouldKeepThisNode)
-			{
-				MergeNodesWithParent(Child);
-			}
-
-			// Go down to children even if we just merged all its children together, because
-			// it is very likely that part of its hierarchy was not processed because it was
-			// marked as separate geometry
-			RecursiveMergeScene(Child);
-		}
-	}
-
-	/** Merge all children of particular node into a single geometry */
-	void MergeNodesWithParent(TSharedPtr<FDatasmithFBXSceneNode>& InParentNode)
-	{
-		TArray<FNodeInfo> AllMatchingNodeInfos;
-		// Look for nodes to merge
-		RecursiveFindNodesToMerge(InParentNode, AllMatchingNodeInfos);
-		// Sort nodes by hash to improve correlation between different scene parts
-		AllMatchingNodeInfos.Sort();
-
-		if (AllMatchingNodeInfos.Num() >= 2)
-		{
-			NodeInfos.Empty();
-			int32 NodeInfosTotalVertexCount = 0;
-			int32 NodeInfosTotalTriangleCount = 0;
-
-			for (auto& NodeIfo: AllMatchingNodeInfos)
-			{
-				TSharedPtr<FDatasmithFBXSceneNode>& Node = NodeIfo.Node;
-				TSharedPtr<FDatasmithFBXSceneMesh>& Mesh = Node->Mesh;
-
-				int32 VertexCount = Mesh->MeshDescription.Vertices().Num();
-				int32 TriangleCount = DatasmithMeshHelper::GetTriangleCount(Mesh->MeshDescription);
-
-				// don't make too huge meshes(e.g. overflowing)
-				if (NodeInfosTotalVertexCount + VertexCount > MAX_MERGE_VERTEX_COUNT
-					|| NodeInfosTotalTriangleCount + TriangleCount > MAX_MERGE_TRIANGLE_COUNT)
-				{
-					if (NodeInfos.Num() > 1)
-					{
-						MergeMeshNodes(InParentNode);
-					}
-
-					NodeInfos.Empty();
-					NodeInfosTotalVertexCount = 0;
-					NodeInfosTotalTriangleCount = 0;
-				}
-
-				NodeInfos.Add(NodeIfo);
-				NodeInfosTotalVertexCount += VertexCount;
-				NodeInfosTotalTriangleCount += TriangleCount;
-			}
-
-			if (NodeInfos.Num() > 1)
-			{
-				MergeMeshNodes(InParentNode);
-			}
-		}
-	}
-
-	struct FNodeInfo
-	{
-		TSharedPtr<FDatasmithFBXSceneNode> Node;
-		FTransform Transform;
-
-		FNodeInfo(TSharedPtr<FDatasmithFBXSceneNode>& InNode)
-			: Node(InNode)
-		{
-		}
-
-		const FMD5Hash& GetHash() const
-		{
-			if (!Hash.IsValid())
-			{
-				FMD5 Md5;
-				// Hash mesh
-				FDatasmithFBXHashUtils::UpdateHash(Md5, Node->Mesh->GetHash());
-				// Hash materials
-				for (int32 i = 0; i < Node->Materials.Num(); i++)
-				{
-					const TSharedPtr<FDatasmithFBXSceneMaterial>& Material = Node->Materials[i];
-					FDatasmithFBXHashUtils::UpdateHash(Md5, Material->Name);
-				}
-				// Hash transform - use node's local transform for better correlation
-				FDatasmithFBXHashUtils::UpdateHash(Md5, Node->LocalTransform);
-				// Finalize hash computation
-				Hash.Set(Md5);
-			}
-			return Hash;
-		}
-
-		// Helper function for sorting
-		bool operator<(const FNodeInfo& Other) const
-		{
-			return GetHash() < Other.GetHash();
-		}
-
-	protected:
-		mutable FMD5Hash Hash;
-	};
-
-	/** Find all nodes which should be merged together, stopping a 'marked' nodes */
-	void RecursiveFindNodesToMerge(TSharedPtr<FDatasmithFBXSceneNode>& Node, TArray<FNodeInfo>& OutNodeInfos)
-	{
-		if (Node->Mesh.IsValid())
-		{
-			// This node has a mesh, remember it
-			OutNodeInfos.Add(FNodeInfo(Node));
-		}
-
-		for (int32 NodeIndex = 0; NodeIndex < Node->Children.Num(); NodeIndex++)
-		{
-			TSharedPtr<FDatasmithFBXSceneNode>& Child = Node->Children[NodeIndex];
-
-			// Skip marked nodes, they should be merged separately
-			if (!Child->bShouldKeepThisNode)
-			{
-				RecursiveFindNodesToMerge(Child, OutNodeInfos);
-			}
-		}
-	}
-
-	/** Analyze usefulness of transform when selecting a node for being pivot. Smaller values
-	 *  indicates a better transform. "Ideal" transform has weight 0.
-	 */
-	static float GetTransformWeight(const FTransform& Transform)
-	{
-		/** Warning: This function needs to be resilient against floating point and conversion errors in
-			the transform values. This because we use it's return value to pick the pivot transform for reused merged
-			meshes.
-
-			If the return value varies too much with small changes (e.g. 1 degree of rotation), it might
-			cause us to pick different transforms for different instances of the reused mesh, leading to errors, given
-			that we bake the pivot transform into the vertex positions of the merged mesh
-			(see MergeMeshes, search for ".Transform")
-		*/
-
-		// Best transform has identity scale. Using a large scale for pivot will force other meshes to
-		// be downscaled. The same applies to small scale, i.e. large and small scales are equally bad.
-		float Determinant = FMath::Abs(Transform.GetDeterminant());
-		if (Determinant < 1.0f)
-		{
-			// Small scale
-			Determinant = 1.0f / Determinant;
-		}
-		// Now Determinant is >= 1.0f. Make identity scale returning zero value.
-		float ScaleWeight = FMath::Pow(Determinant, 1.0f/3.0f) - 1.0f;
-
-		// Measure scale uniformness
-		FVector Scale = Transform.GetScale3D().GetAbs();
-		float Uniformness = ( FMath::Abs(Scale.X - Scale.Y) + FMath::Abs(Scale.X - Scale.Z) + FMath::Abs(Scale.Y - Scale.Z) + 0.1f ) / (Scale.X + Scale.Y + Scale.Z + 0.1f);
-		ScaleWeight += Uniformness * 10;
-
-		// Now analyze rotation. Best transforms will have rotation which are multiple of 90 degrees.
-		FRotator Rotation = Transform.GetRotation().Rotator().GetDenormalized();
-		float RotValues[3];
-		RotValues[0] = Rotation.Yaw;
-		RotValues[1] = Rotation.Pitch;
-		RotValues[2] = Rotation.Roll;
-		float RotationWeight = 0.0f;
-		for (float Value : RotValues)
-		{
-			// This value is in range 0..360 now due to GetDenormalized call above
-			Value = FMath::Fmod(Value, 90.0f);
-			float AxisWeight = (Value > 45.0f) ? 90.0f - Value : Value;
-			RotationWeight += AxisWeight / 45.0f;
-
-			// Huge tolerance of 1 degree since it's common for rotations to go through several conversions
-			// and end up with huge compounded errors
-			if (!FMath::IsNearlyZero(Value, 1))
-			{
-				// Non-zero rotation is worse than zero
-				RotationWeight += 10;
-			}
-		}
-
-		return ScaleWeight + RotationWeight;
-	}
-
-	/** Merge all found mesh nodes into a single mesh node */
-	void MergeMeshNodes(TSharedPtr<FDatasmithFBXSceneNode>& InParentNode)
-	{
-		// Compute parent-related transforms for each node
-		for (FNodeInfo& NodeInfo: NodeInfos)
-		{
-			NodeInfo.Transform = NodeInfo.Node->GetTransformRelativeToParent(InParentNode);
-		}
-
-		// Use one of mesh nodes as "pivot" for node group. All other nodes will be repositioned as
-		// children of that node. This allows us to avoid any differences in parent node transforms, for
-		// example when we have the 2 instances of the same node group with and without mirror transform
-		// in hierarchy - such groups still will be considered as identical.
-		FTransform PivotTransform;
-		float BestNodeWeight = 0.0f;
-		for (int32 Index = 0; Index < NodeInfos.Num(); Index++)
-		{
-			const FTransform& Transform = NodeInfos[Index].Transform;
-			float Weight = GetTransformWeight(Transform);
-			if (Index == 0 || Weight < BestNodeWeight)
-			{
-				PivotTransform = Transform;
-				BestNodeWeight = Weight;
-			}
-		}
-		// Reposition nodes
-		FTransform InversePivotTransform = PivotTransform.Inverse();
-		for (FNodeInfo& NodeInfo : NodeInfos)
-		{
-			NodeInfo.Transform = NodeInfo.Transform * InversePivotTransform;
-		}
-
-		// Compute hash of the node set
-		FMD5 Md5;
-		for (FNodeInfo& NodeInfo : NodeInfos)
-		{
-			FDatasmithFBXHashUtils::UpdateHash(Md5, NodeInfo.GetHash());
-		}
-		FMD5Hash NodeSetHash;
-		NodeSetHash.Set(Md5);
-
-		FString NewNodeName = FString::Printf(TEXT("%s") MERGED_SUFFIX, *InParentNode->Name);
-
-//		UE_LOG(LogDatasmithFBXImport, Log, TEXT("Merging %d meshes into node \"%s\", hash = %s"), NodeInfos.Num(), *NewNodeName, *Lex::ToString(NodeSetHash));
-
-		// Make a node for the new mesh
-		TSharedPtr<FDatasmithFBXSceneNode> NewNode(new FDatasmithFBXSceneNode());
-		NewNode->Name = NewNodeName;
-		NewNode->Parent = InParentNode;
-		NewNode->LocalTransform = PivotTransform;
-		NewNode->OriginalName = NewNode->Name; // So that it becomes a Tag later
-		InParentNode->Children.Add(NewNode);
-
-		// Find if we already have identical set of nodes merged
-		const TSharedPtr<FDatasmithFBXSceneNode>* PreviousMergedNode = MergedNodes.Find(NodeSetHash);
-		if (PreviousMergedNode == nullptr)
-		{
-			// Make an empty mesh
-			TSharedPtr<FDatasmithFBXSceneMesh> NewMesh(new FDatasmithFBXSceneMesh());
-			NewMesh->Name = NewNodeName;
-			NewNode->Mesh = NewMesh;
-
-			// Merge geometries
-			MergeMeshes(NewNode);
-
-			// Remember node and mesh for later reuse when needed
-			MergedNodes.Add(NodeSetHash, NewNode);
-		}
-		else
-		{
-//			UE_LOG(LogDatasmithFBXImport, Log, TEXT("... reusing mesh"));
-			// Do not merge if the same mesh set already exists
-			NewNode->Mesh = (*PreviousMergedNode)->Mesh;
-			// Copy materials too
-			NewNode->Materials = (*PreviousMergedNode)->Materials;
-		}
-
-		// Release merged meshes and materials
-		for (FNodeInfo& NodeInfo : NodeInfos)
-		{
-			TSharedPtr<FDatasmithFBXSceneNode>& Node = NodeInfo.Node;
-			Node->Mesh.Reset();
-			Node->Materials.Empty();
-		}
-	}
-
-	void MergeMeshes(TSharedPtr<FDatasmithFBXSceneNode>& MergedNode)
-	{
-		if (!ensure(MergedNode && MergedNode->Mesh))
-		{
-			return;
-		}
-
-		int32 MergedVertexCount = 0;
-		int32 MergedVertexInstanceCount = 0;
-		int32 MergedEdgeCount = 0;
-		int32 MergedPolygonCount = 0;
-
-		TArray< TSharedPtr<FDatasmithFBXSceneMaterial> >& TargetMaterials = MergedNode->Materials;
-		for (int32 MeshIndex = 0; MeshIndex < NodeInfos.Num(); MeshIndex++)
-		{
-			TSharedPtr<FDatasmithFBXSceneNode>& Node = NodeInfos[MeshIndex].Node;
-			if (!ensure(Node && Node->Mesh))
-			{
-				return;
-			}
-			FMeshDescription& SourceMeshDescription = Node->Mesh->MeshDescription;
-
-			MergedVertexCount += SourceMeshDescription.Vertices().Num();
-			MergedVertexInstanceCount += SourceMeshDescription.VertexInstances().Num();
-			MergedEdgeCount += SourceMeshDescription.Edges().Num();
-			MergedPolygonCount += SourceMeshDescription.Polygons().Num();
-
-			// Merge materials
-			for (int32 Index = 0; Index < Node->Materials.Num(); Index++)
-			{
-				TargetMaterials.AddUnique(Node->Materials[Index]);
-			}
-		}
-
-		// Prepare destination mesh
-		TSharedPtr<FDatasmithFBXSceneMesh> MergedMesh = MergedNode->Mesh;
-		FMeshDescription& MergedMeshDescription = MergedMesh->MeshDescription;
-		DatasmithMeshHelper::PrepareAttributeForStaticMesh(MergedMeshDescription);
-		MergedMeshDescription.ReserveNewPolygonGroups(TargetMaterials.Num());
-		MergedMeshDescription.ReserveNewVertexInstances(MergedVertexCount);
-		MergedMeshDescription.ReserveNewEdges(MergedVertexInstanceCount);
-		MergedMeshDescription.ReserveNewVertices(MergedEdgeCount);
-		MergedMeshDescription.ReserveNewPolygons(MergedPolygonCount);
-
-		TPolygonGroupAttributesRef<FName> PolygonGroupImportedMaterialSlotNames = MergedMeshDescription.PolygonGroupAttributes().GetAttributesRef<FName>(MeshAttribute::PolygonGroup::ImportedMaterialSlotName);
-
-		// Now merge meshes
-		for (int32 MeshIndex = 0; MeshIndex < NodeInfos.Num(); MeshIndex++)
-		{
-			FNodeInfo& NodeInfo = NodeInfos[MeshIndex];
-			TSharedPtr<FDatasmithFBXSceneNode>& Node = NodeInfo.Node;
-			FMeshDescription& SourceMeshDescription = Node->Mesh->MeshDescription;
-
-			// Remap from material indices since we might or might not add new materials to target
-			// given that we don't allow repeated entries
-			// Should we? If we don't the user can't individually reassign materials to segments of the mesh that
-			// were originally distinct...
-			const TArray<TSharedPtr<FDatasmithFBXSceneMaterial>>& SourceMaterials = Node->Materials;
-			TArray<int32> SourceToTargetMatID;
-			if (SourceMaterials.Num() > 0)
-			{
-				for (int32 SourceIndex = 0; SourceIndex < SourceMaterials.Num(); SourceIndex++)
-				{
-					int32 TargetIndex = TargetMaterials.Find(SourceMaterials[SourceIndex]);
-					check(TargetIndex >= 0);
-					SourceToTargetMatID.Add(TargetIndex);  // SourceToTargetMatID[SourceIndex] -> TargetIndex
-				}
-			}
-			else
-			{
-				SourceToTargetMatID.Add(0);
-			}
-
-			// Expand target mesh polygon groups to fit as many different groups as we need to support our new source mesh
-			while (MergedMeshDescription.PolygonGroups().Num() - 1 < FMath::Max(SourceToTargetMatID))
-			{
-				FPolygonGroupID NewGroup = MergedMeshDescription.CreatePolygonGroup();
-				PolygonGroupImportedMaterialSlotNames[NewGroup] = DatasmithMeshHelper::DefaultSlotName(NewGroup.GetValue());
-			}
-
-			FMeshDescriptionOperations::FAppendSettings AppendSettings;
-
-			bool bRecalculateNormals = false;
-			const FTransform& Transform = NodeInfo.Transform;
-			if (Transform.IsValid())
-			{
-				AppendSettings.MeshTransform = Transform;
-
-				// Odd negative scales on the *node* transform don't affect winding order as the mesh normals/faces remain
-				// pointing to the same direction (e.g. always outward). Exception to that is when we'll bake the scaling
-				// in the vertex positions, which is just about to happen. In here, we flip triangle winding order
-				// to compensate for that
-				if (Transform.GetDeterminant() < 0.0f && !Node->Mesh->bFlippedFaces)
-				{
-					SourceMeshDescription.ReverseAllPolygonFacing();
-					Node->Mesh->bFlippedFaces = true;
-
-					bRecalculateNormals = true;
-				}
-				// We need to check whether the mesh has flipped faces or not, though.
-				// Example: Two wheels (same mesh), one with negative scale. If the negative scale one is merged first,
-				// the above if will flip the faces. When the non-negative scale wheel is merged, it will remain with
-				// flipped faces (equal meshes are shared on import). This reverses that flip
-				else if (Transform.GetDeterminant() > 0.0f && Node->Mesh->bFlippedFaces)
-				{
-					SourceMeshDescription.ReverseAllPolygonFacing();
-					Node->Mesh->bFlippedFaces = false;
-
-					bRecalculateNormals = true;
-				}
-			}
-
-			// RecomputeNormalsAndTangentsIfNeeded expects the normals, tangents and binormals arrays to be allocated,
-			// so we need to make sure those exist, or else we will have no opportunity to fix our normals and tangents before
-			// merging
-			TPolygonAttributesConstRef<FVector> PolygonNormals = SourceMeshDescription.PolygonAttributes().GetAttributesRef<FVector>(MeshAttribute::Polygon::Normal);
-			TPolygonAttributesConstRef<FVector> PolygonTangents = SourceMeshDescription.PolygonAttributes().GetAttributesRef<FVector>(MeshAttribute::Polygon::Tangent);
-			TPolygonAttributesConstRef<FVector> PolygonBinormals = SourceMeshDescription.PolygonAttributes().GetAttributesRef<FVector>(MeshAttribute::Polygon::Binormal);
-			if (!PolygonNormals.IsValid())
-			{
-				SourceMeshDescription.PolygonAttributes().RegisterAttribute<FVector>(MeshAttribute::Polygon::Normal, 1, FVector::ZeroVector);
-			}
-			if (!PolygonTangents.IsValid())
-			{
-				SourceMeshDescription.PolygonAttributes().RegisterAttribute<FVector>(MeshAttribute::Polygon::Tangent, 1, FVector::ZeroVector);
-			}
-			if (!PolygonBinormals.IsValid())
-			{
-				SourceMeshDescription.PolygonAttributes().RegisterAttribute<FVector>(MeshAttribute::Polygon::Binormal, 1, FVector::ZeroVector);
-			}
-
-			//Use MikktSpace
-			uint32 TangentOptions = FMeshDescriptionOperations::ETangentOptions::UseMikkTSpace;
-			//Use weighted normals
-			//TangentOptions |= FMeshDescriptionOperations::ETangentOptions::UseWeightedAreaAndAngle;
-			FMeshDescriptionOperations::RecomputeNormalsAndTangentsIfNeeded(SourceMeshDescription, (FMeshDescriptionOperations::ETangentOptions)TangentOptions, bRecalculateNormals, bRecalculateNormals);
-
-			// The default merge process (FMeshDescriptionOperations::AppendMeshDescription) will merge polygon groups if
-			// the slot names are the same. So if mesh A has just material slots 0 and 1 with default names, and mesh B has just slot 0,
-			// also with a default name, we will get a mesh C with slots 0 and 1 with default names.
-			// Assuming the materials are all discrete however, we would want instead to keep the polygon groups separate,
-			// so that we can have discrete slots 0, 1 and 2
-			AppendSettings.PolygonGroupsDelegate.BindLambda([&SourceToTargetMatID](const FMeshDescription& SourceMesh, FMeshDescription& TargetMesh, PolygonGroupMap& RemapPolygonGroup)
-			{
-				for (FPolygonGroupID SourcePolygonGroupID : SourceMesh.PolygonGroups().GetElementIDs())
-				{
-					RemapPolygonGroup.Add(SourcePolygonGroupID, FPolygonGroupID(SourceToTargetMatID[SourcePolygonGroupID.GetValue()]));
-				}
-			});
-
-			FMeshDescriptionOperations::AppendMeshDescription(SourceMeshDescription, MergedMeshDescription, AppendSettings);
-		}
-	}
-
-	/** List of nodes which has meshes and which should be merged together */
-	TArray<FNodeInfo> NodeInfos;
-
-	/** Map between hash of node set and combined node, to reuse previously combined meshes if node sets matches */
-	TMap< FMD5Hash, TSharedPtr<FDatasmithFBXSceneNode> > MergedNodes;
-};
-
-void FDatasmithFBXSceneProcessor::MergeSceneNodes()
-{
-	FNodeMergeHelper Helper;
-	Helper.MergeScene(Scene);
-}
-
 void FDatasmithFBXSceneProcessor::FixNodeNames()
 {
 	struct FFixNodeNamesHelper
diff --git a/Engine/Plugins/Enterprise/DatasmithFBXImporter/Source/DatasmithFBXTranslator/Public/DatasmithFBXSceneProcessor.h b/Engine/Plugins/Enterprise/DatasmithFBXImporter/Source/DatasmithFBXTranslator/Public/DatasmithFBXSceneProcessor.h
--- a/Engine/Plugins/Enterprise/DatasmithFBXImporter/Source/DatasmithFBXTranslator/Public/DatasmithFBXSceneProcessor.h
+++ b/Engine/Plugins/Enterprise/DatasmithFBXImporter/Source/DatasmithFBXTranslator/Public/DatasmithFBXSceneProcessor.h
@@ -42,9 +42,6 @@ public:
 	/** Find and mark scene nodes which shouldn't be merged with parent or sibling nodes */
 	void FindPersistentNodes();
 
-	/** Merge scene nodes together. Nodes which were marked with bShouldKeepThisNode will be preserved. */
-	void MergeSceneNodes();
-
 	/** Find duplicated nodes which could be replaced with a single instance and pointer to children ("shared nodes") */
 	void OptimizeDuplicatedNodes();
 
diff --git a/Engine/Plugins/Enterprise/DatasmithFBXImporter/Source/DatasmithVREDTranslator/Private/DatasmithVREDImportOptions.cpp b/Engine/Plugins/Enterprise/DatasmithFBXImporter/Source/DatasmithVREDTranslator/Private/DatasmithVREDImportOptions.cpp
--- a/Engine/Plugins/Enterprise/DatasmithFBXImporter/Source/DatasmithVREDTranslator/Private/DatasmithVREDImportOptions.cpp
+++ b/Engine/Plugins/Enterprise/DatasmithFBXImporter/Source/DatasmithVREDTranslator/Private/DatasmithVREDImportOptions.cpp
@@ -13,7 +13,6 @@
 
 UDatasmithVREDImportOptions::UDatasmithVREDImportOptions(const FObjectInitializer& ObjectInitializer)
 	: Super(ObjectInitializer)
-	, bMergeNodes(false)
 	, bOptimizeDuplicatedNodes(false)
 	, bImportMats(true)
 	, bImportVar(true)
@@ -88,7 +87,6 @@ void UDatasmithVREDImportOptions::FromSceneImportData(UDatasmithFBXSceneImportDa
 
 	if (UDatasmithVREDSceneImportData* VREDImportData = Cast<UDatasmithVREDSceneImportData>(InImportData))
 	{
-		bMergeNodes					= VREDImportData->bMergeNodes;
 		bOptimizeDuplicatedNodes	= VREDImportData->bOptimizeDuplicatedNodes;
 		bImportMats					= VREDImportData->bImportMats;
 		MatsPath.FilePath			= VREDImportData->MatsPath;
@@ -108,7 +106,6 @@ void UDatasmithVREDImportOptions::ToSceneImportData(UDatasmithFBXSceneImportData
 
 	if (UDatasmithVREDSceneImportData* VREDImportData = Cast<UDatasmithVREDSceneImportData>(OutImportData))
 	{
-		VREDImportData->bMergeNodes					= bMergeNodes;
 		VREDImportData->bOptimizeDuplicatedNodes	= bOptimizeDuplicatedNodes;
 		VREDImportData->bImportMats					= bImportMats;
 		VREDImportData->MatsPath					= MatsPath.FilePath;
diff --git a/Engine/Plugins/Enterprise/DatasmithFBXImporter/Source/DatasmithVREDTranslator/Private/DatasmithVREDImportOptions.h b/Engine/Plugins/Enterprise/DatasmithFBXImporter/Source/DatasmithVREDTranslator/Private/DatasmithVREDImportOptions.h
--- a/Engine/Plugins/Enterprise/DatasmithFBXImporter/Source/DatasmithVREDTranslator/Private/DatasmithVREDImportOptions.h
+++ b/Engine/Plugins/Enterprise/DatasmithFBXImporter/Source/DatasmithVREDTranslator/Private/DatasmithVREDImportOptions.h
@@ -19,9 +19,6 @@ class UDatasmithVREDImportOptions : public UDatasmithFBXImportOptions
 	GENERATED_UCLASS_BODY()
 
 public:
-	UPROPERTY(config, EditAnywhere, BlueprintReadWrite, Category=Processing, meta=(ToolTip="Merge nodes that are not used in variants/animations"))
-	bool bMergeNodes;
-
 	UPROPERTY(config, EditAnywhere, BlueprintReadWrite, Category=Processing, meta=(ToolTip="Finds duplicated scene nodes which could be replaced with a single instance"))
 	bool bOptimizeDuplicatedNodes;
 
diff --git a/Engine/Plugins/Enterprise/DatasmithFBXImporter/Source/DatasmithVREDTranslator/Private/DatasmithVREDImporter.cpp b/Engine/Plugins/Enterprise/DatasmithFBXImporter/Source/DatasmithVREDTranslator/Private/DatasmithVREDImporter.cpp
--- a/Engine/Plugins/Enterprise/DatasmithFBXImporter/Source/DatasmithVREDTranslator/Private/DatasmithVREDImporter.cpp
+++ b/Engine/Plugins/Enterprise/DatasmithFBXImporter/Source/DatasmithVREDTranslator/Private/DatasmithVREDImporter.cpp
@@ -775,14 +775,6 @@ void FDatasmithVREDImporter::ProcessScene()
 		Processor.OptimizeDuplicatedNodes();
 	}
 
-	if (ImportOptions->bMergeNodes)
-	{
-		Processor.MergeSceneNodes();
-
-		//Call "RemoveEmptyNodes" again after merging, because all nodes which previously had meshes are empty now
-		Processor.RemoveEmptyNodes();
-	}
-
 	Processor.FixMeshNames();
 }
 
