Multithread FBX Mesh conversion to MeshDescription
  - 2m23s -> 21s when importing DATASET-0022a

From dev-enterprise[REDACTED_HANDLE]

diff --git a/Engine/Plugins/Enterprise/DatasmithFBXImporter/Source/DatasmithFBXTranslator/Private/DatasmithFBXFileImporter.cpp b/Engine/Plugins/Enterprise/DatasmithFBXImporter/Source/DatasmithFBXTranslator/Private/DatasmithFBXFileImporter.cpp
--- a/Engine/Plugins/Enterprise/DatasmithFBXImporter/Source/DatasmithFBXTranslator/Private/DatasmithFBXFileImporter.cpp
+++ b/Engine/Plugins/Enterprise/DatasmithFBXImporter/Source/DatasmithFBXTranslator/Private/DatasmithFBXFileImporter.cpp
@@ -8,6 +8,7 @@
 #include "DatasmithMeshHelper.h"
 #include "DatasmithUtils.h"
 
+#include "Async/Async.h"
 #include "Curves/RichCurve.h"
 #include "FbxImporter.h"
 #include "StaticMeshAttributes.h"
@@ -82,6 +83,8 @@ FDatasmithFBXFileImporter::FDatasmithFBXFileImporter(FbxScene* InFbxScene, FData
 
 void FDatasmithFBXFileImporter::ImportScene()
 {
+	TRACE_CPUPROFILER_EVENT_SCOPE(FDatasmithFBXFileImporter::ImportScene)
+
 	ConvertCoordinateSystem();
 
 	// ImportMaterials
@@ -134,10 +137,30 @@ void FDatasmithFBXFileImporter::ImportScene()
 	// Import scene hierarchy, extracting animations and creating placeholders for all meshes
 	TraverseHierarchyNodeRecursively(InScene->GetRootNode(), OutScene->RootNode);
 
+	TArray<TFuture<FbxMesh*>> Tasks;
+	Tasks.Reserve(ImportedMeshes.Num());
+
 	// Import meshes
 	for (auto& It : ImportedMeshes)
 	{
-		DoImportMesh(It.Key, It.Value.Get());
+		Tasks.Emplace(
+			Async(
+				EAsyncExecution::LargeThreadPool,
+				[FbxMeshPtr = It.Key, DatasmithFBXSceneMeshPtr = It.Value.Get()]()
+				{
+					DoImportMesh(FbxMeshPtr, DatasmithFBXSceneMeshPtr);
+					return FbxMeshPtr;
+				}
+			)
+		);
+	}
+
+	for (TFuture<FbxMesh*>& Future : Tasks)
+	{
+		// Destroy Fbx mesh to save memory - it won't be used anymore. Destroying an FbxMesh will unlink it from all places
+		// in Fbx scene, i.e. all instances will be lost. However we're calling DoImportMesh() when we already have full
+		// Fbx hierarchy imported into scene structures, so it won't damage anything.
+		Future.Get()->Destroy();
 	}
 }
 
@@ -892,6 +915,7 @@ TSharedPtr<FDatasmithFBXSceneMesh> FDatasmithFBXFileImporter::ImportMesh(FbxMesh
 void FDatasmithFBXFileImporter::DoImportMesh(FbxMesh* InMesh, FDatasmithFBXSceneMesh* Mesh)
 {
 	// Reference code: UnFbx::FFbxImporter::BuildStaticMeshFromGeometry()
+	TRACE_CPUPROFILER_EVENT_SCOPE(FDatasmithFBXFileImporter::DoImportMesh)
 
 	int32 MaterialCount = FMath::Max(Mesh->ImportMaterialCount, 1);
 
@@ -1206,10 +1230,6 @@ void FDatasmithFBXFileImporter::DoImportMesh(FbxMesh* InMesh, FDatasmithFBXScene
 		const FPolygonID NewPolygonID = MeshDescription.CreatePolygon(PolygonGroupID, CornerVertexInstanceIDs);
 	} // for poly loop
 
-	// Destroy Fbx mesh to save memory - it won't be used anymore. Destroying an FbxMesh will unlink it from all places
-	// in Fbx scene, i.e. all instances will be lost. However we're calling DoImportMesh() when we already have full
-	// Fbx hierarchy imported into scene structures, so it won't damage anything.
-	InMesh->Destroy();
 	Mesh->MeshDescription = MoveTemp(MeshDescription);
 }
 
diff --git a/Engine/Plugins/Enterprise/DatasmithFBXImporter/Source/DatasmithFBXTranslator/Public/DatasmithFBXFileImporter.h b/Engine/Plugins/Enterprise/DatasmithFBXImporter/Source/DatasmithFBXTranslator/Public/DatasmithFBXFileImporter.h
--- a/Engine/Plugins/Enterprise/DatasmithFBXImporter/Source/DatasmithFBXTranslator/Public/DatasmithFBXFileImporter.h
+++ b/Engine/Plugins/Enterprise/DatasmithFBXImporter/Source/DatasmithFBXTranslator/Public/DatasmithFBXFileImporter.h
@@ -43,12 +43,12 @@ protected:
 	TSharedPtr<FDatasmithFBXSceneMesh> ImportMesh(FbxMesh* InMesh, FbxNode* InNode);
 
 	/** Convert Fbx mesh into FRawMesh */
-	void DoImportMesh(FbxMesh* InMesh, FDatasmithFBXSceneMesh* Mesh);
+	static void DoImportMesh(FbxMesh* InMesh, FDatasmithFBXSceneMesh* Mesh);
 
 	/** Import an Fbx material */
 	TSharedPtr<FDatasmithFBXSceneMaterial> ImportMaterial(FbxSurfaceMaterial* InMaterial);
 
-	void FindFbxUVChannels(FbxMesh* Mesh, TArray<FFbxUVInfo>& FbxUVs);
+	static void FindFbxUVChannels(FbxMesh* Mesh, TArray<FFbxUVInfo>& FbxUVs);
 
 	bool IsOddNegativeScale(FbxAMatrix& TotalMatrix);
 
