Remove duplicate lightmap UV import settings #jira UE-82356 #rb JeanMichel.Dignard

diff --git a/Engine/Plugins/Enterprise/DatasmithFBXImporter/Source/DatasmithDeltaGenTranslator/Private/DatasmithDeltaGenImportOptions.cpp b/Engine/Plugins/Enterprise/DatasmithFBXImporter/Source/DatasmithDeltaGenTranslator/Private/DatasmithDeltaGenImportOptions.cpp
--- a/Engine/Plugins/Enterprise/DatasmithFBXImporter/Source/DatasmithDeltaGenTranslator/Private/DatasmithDeltaGenImportOptions.cpp
+++ b/Engine/Plugins/Enterprise/DatasmithFBXImporter/Source/DatasmithDeltaGenTranslator/Private/DatasmithDeltaGenImportOptions.cpp
@@ -21,7 +21,6 @@ UDatasmithDeltaGenImportOptions::UDatasmithDeltaGenImportOptions(const FObjectIn
 	, bImportTml(true)
 	, ShadowTextureMode(EShadowTextureMode::Ignore)
 {
-	bGenerateLightmapUVs = false;
 }
 
 namespace DeltaGenImportOptionsImpl
diff --git a/Engine/Plugins/Enterprise/DatasmithFBXImporter/Source/DatasmithDeltaGenTranslator/Private/DatasmithDeltaGenImporter.cpp b/Engine/Plugins/Enterprise/DatasmithFBXImporter/Source/DatasmithDeltaGenTranslator/Private/DatasmithDeltaGenImporter.cpp
--- a/Engine/Plugins/Enterprise/DatasmithFBXImporter/Source/DatasmithDeltaGenTranslator/Private/DatasmithDeltaGenImporter.cpp
+++ b/Engine/Plugins/Enterprise/DatasmithFBXImporter/Source/DatasmithDeltaGenTranslator/Private/DatasmithDeltaGenImporter.cpp
@@ -445,20 +445,16 @@ TSharedPtr<IDatasmithActorElement> FDatasmithDeltaGenImporter::ConvertNode(const
 			MeshNameToFBXMesh.Add(MeshName, ThisMesh);
 			TSharedRef<IDatasmithMeshElement> MeshElement = FDatasmithSceneFactory::CreateMesh(*ThisMesh->Name);
 
-			if (ImportOptions->bGenerateLightmapUVs)
+			FMeshDescription& MeshDescription = ThisMesh->MeshDescription;
+			FStaticMeshAttributes StaticMeshAttributes(MeshDescription);
+			TVertexInstanceAttributesRef<FVector2D> VertexInstanceUVs = StaticMeshAttributes.GetVertexInstanceUVs();
+			int32 NumUVChannels = VertexInstanceUVs.GetNumIndices();
+
+			// DeltaGen uses UV channel 0 for texture UVs, and UV channel 1 for lightmap UVs
+			// Don't set it to zero or else it will disable Datasmith's GenerateLightmapUV option
+			if (NumUVChannels > 1)
 			{
-				MeshElement->SetLightmapSourceUV(0);
-				MeshElement->SetLightmapCoordinateIndex(-1);
-			}
-			else
-			{
-				FMeshDescription& MeshDescription = ThisMesh->MeshDescription;
-				FStaticMeshAttributes StaticMeshAttributes(MeshDescription);
-				TVertexInstanceAttributesRef<FVector2D> VertexInstanceUVs = StaticMeshAttributes.GetVertexInstanceUVs();
-				int32 NumUVChannels = VertexInstanceUVs.GetNumIndices();
-
-				// DeltaGen uses UV channel 0 for texture UVs, and UV channel 1 for lightmap UVs
-				MeshElement->SetLightmapCoordinateIndex(NumUVChannels > 1 ? 1 : 0);
+				MeshElement->SetLightmapCoordinateIndex(1);
 			}
 
 			DatasmithScene->AddMesh(MeshElement);
diff --git a/Engine/Plugins/Enterprise/DatasmithFBXImporter/Source/DatasmithFBXTranslator/Private/DatasmithFBXImportOptions.cpp b/Engine/Plugins/Enterprise/DatasmithFBXImporter/Source/DatasmithFBXTranslator/Private/DatasmithFBXImportOptions.cpp
--- a/Engine/Plugins/Enterprise/DatasmithFBXImporter/Source/DatasmithFBXTranslator/Private/DatasmithFBXImportOptions.cpp
+++ b/Engine/Plugins/Enterprise/DatasmithFBXImporter/Source/DatasmithFBXTranslator/Private/DatasmithFBXImportOptions.cpp
@@ -10,7 +10,6 @@
 
 UDatasmithFBXImportOptions::UDatasmithFBXImportOptions(const FObjectInitializer& ObjectInitializer)
 	: Super(ObjectInitializer)
-	, bGenerateLightmapUVs(false)
 	, IntermediateSerialization(EDatasmithFBXIntermediateSerializationType::Disabled)
 	, bColorizeMaterials(false)
 {
@@ -20,9 +19,7 @@ void UDatasmithFBXImportOptions::FromSceneImportData(UDatasmithFBXSceneImportDat
 {
 	if (InImportData)
 	{
-		bGenerateLightmapUVs		= InImportData->bGenerateLightmapUVs;
-		//TexturesDir.Path			= InImportData->TexturesDir;
-		IntermediateSerialization	= (EDatasmithFBXIntermediateSerializationType)InImportData->IntermediateSerialization;
+		IntermediateSerialization = (EDatasmithFBXIntermediateSerializationType)InImportData->IntermediateSerialization;
 	}
 }
 
@@ -30,9 +27,7 @@ void UDatasmithFBXImportOptions::ToSceneImportData(UDatasmithFBXSceneImportData*
 {
 	if (OutImportData)
 	{
-		OutImportData->bGenerateLightmapUVs			= bGenerateLightmapUVs;
-		//OutImportData->TexturesDir					= TexturesDir.Path;
-		OutImportData->IntermediateSerialization	= (uint8)IntermediateSerialization;
+		OutImportData->IntermediateSerialization = (uint8)IntermediateSerialization;
 	}
 }
 
diff --git a/Engine/Plugins/Enterprise/DatasmithFBXImporter/Source/DatasmithFBXTranslator/Public/DatasmithFBXImportOptions.h b/Engine/Plugins/Enterprise/DatasmithFBXImporter/Source/DatasmithFBXTranslator/Public/DatasmithFBXImportOptions.h
--- a/Engine/Plugins/Enterprise/DatasmithFBXImporter/Source/DatasmithFBXTranslator/Public/DatasmithFBXImportOptions.h
+++ b/Engine/Plugins/Enterprise/DatasmithFBXImporter/Source/DatasmithFBXTranslator/Public/DatasmithFBXImportOptions.h
@@ -25,9 +25,6 @@ class DATASMITHFBXTRANSLATOR_API UDatasmithFBXImportOptions : public UObject
 	GENERATED_UCLASS_BODY()
 
 public:
-	UPROPERTY(config, EditAnywhere, BlueprintReadWrite, Category = Lightmaps, meta = (ToolTip = "Generate new UV coordinates for lightmapping instead of using the highest index UV set. \nTurn this on to have Unreal Studio generate lightmap UV sets automatically.\nTurn this off to try using the highest index existing UV set (if available) as the lightmap UV set.\nFor both cases, geometry without existing UV sets will receive an empty UV set, which will by itself not be valid for use with Lightmass."))
-	bool bGenerateLightmapUVs;
-
 	UPROPERTY(config, EditAnywhere, BlueprintReadWrite, Category=AssetImporting, meta=(DisplayName="Texture folders", ToolTip="Where to look for textures"))
 	TArray<FDirectoryPath> TextureDirs;
 
diff --git a/Engine/Plugins/Enterprise/DatasmithFBXImporter/Source/DatasmithVREDTranslator/Private/DatasmithVREDImportOptions.cpp b/Engine/Plugins/Enterprise/DatasmithFBXImporter/Source/DatasmithVREDTranslator/Private/DatasmithVREDImportOptions.cpp
--- a/Engine/Plugins/Enterprise/DatasmithFBXImporter/Source/DatasmithVREDTranslator/Private/DatasmithVREDImportOptions.cpp
+++ b/Engine/Plugins/Enterprise/DatasmithFBXImporter/Source/DatasmithVREDTranslator/Private/DatasmithVREDImportOptions.cpp
@@ -20,7 +20,6 @@ UDatasmithVREDImportOptions::UDatasmithVREDImportOptions(const FObjectInitialize
 	, bImportLightInfo(true)
 	, bImportClipInfo(true)
 {
-	bGenerateLightmapUVs = false;
 }
 
 namespace VREDImportOptionsImpl
diff --git a/Engine/Plugins/Enterprise/DatasmithFBXImporter/Source/DatasmithVREDTranslator/Private/DatasmithVREDImporter.cpp b/Engine/Plugins/Enterprise/DatasmithFBXImporter/Source/DatasmithVREDTranslator/Private/DatasmithVREDImporter.cpp
--- a/Engine/Plugins/Enterprise/DatasmithFBXImporter/Source/DatasmithVREDTranslator/Private/DatasmithVREDImporter.cpp
+++ b/Engine/Plugins/Enterprise/DatasmithFBXImporter/Source/DatasmithVREDTranslator/Private/DatasmithVREDImporter.cpp
@@ -864,16 +864,6 @@ TSharedPtr<IDatasmithActorElement> FDatasmithVREDImporter::ConvertNode(const TSh
 			MeshNameToFBXMesh.Add(MeshName, ThisMesh);
 			TSharedRef<IDatasmithMeshElement> MeshElement = FDatasmithSceneFactory::CreateMesh(*ThisMesh->Name);
 
-			if (ImportOptions->bGenerateLightmapUVs)
-			{
-				MeshElement->SetLightmapSourceUV(0);
-				MeshElement->SetLightmapCoordinateIndex(-1);
-			}
-			else
-			{
-				MeshElement->SetLightmapCoordinateIndex(0);
-			}
-
 			DatasmithScene->AddMesh(MeshElement);
 		}
 
