DeltaGen: Fix pivot decomposition's handling of animations #jira UE-82557 #rb Anousack.Kitisa, JeanLuc.Corenthin

diff --git a/Engine/Plugins/Enterprise/DatasmithFBXImporter/Source/DatasmithDeltaGenTranslator/Private/DatasmithDeltaGenImportData.h b/Engine/Plugins/Enterprise/DatasmithFBXImporter/Source/DatasmithDeltaGenTranslator/Private/DatasmithDeltaGenImportData.h
--- a/Engine/Plugins/Enterprise/DatasmithFBXImporter/Source/DatasmithDeltaGenTranslator/Private/DatasmithDeltaGenImportData.h
+++ b/Engine/Plugins/Enterprise/DatasmithFBXImporter/Source/DatasmithDeltaGenTranslator/Private/DatasmithDeltaGenImportData.h
@@ -252,23 +252,25 @@ struct FDeltaGenPosData
 };
 
 UENUM(BlueprintType)
-enum class FDeltaGenTmlDataAnimationTrackType : uint8
+enum class EDeltaGenTmlDataAnimationTrackType : uint8
 {
-	Unsupported,
-	Translation,
-	Rotation,
-	RotationDeltaGenEuler,
-	Scale,
-	Center
+	Unsupported = 0,
+	Translation = 1,
+	Rotation = 2,
+	RotationDeltaGenEuler = 4,
+	Scale = 8,
+	Center = 16
 };
 
+ENUM_CLASS_FLAGS(EDeltaGenTmlDataAnimationTrackType);
+
 USTRUCT(BlueprintType)
 struct FDeltaGenTmlDataAnimationTrack
 {
 	GENERATED_USTRUCT_BODY()
 
 	UPROPERTY(BlueprintReadWrite, EditAnywhere, Category=DeltaGen)
-	FDeltaGenTmlDataAnimationTrackType Type;
+	EDeltaGenTmlDataAnimationTrackType Type;
 
 	UPROPERTY(BlueprintReadWrite, EditAnywhere, Category=DeltaGen)
 	TArray<float> Keys;
diff --git a/Engine/Plugins/Enterprise/DatasmithFBXImporter/Source/DatasmithDeltaGenTranslator/Private/DatasmithDeltaGenImporter.cpp b/Engine/Plugins/Enterprise/DatasmithFBXImporter/Source/DatasmithDeltaGenTranslator/Private/DatasmithDeltaGenImporter.cpp
--- a/Engine/Plugins/Enterprise/DatasmithFBXImporter/Source/DatasmithDeltaGenTranslator/Private/DatasmithDeltaGenImporter.cpp
+++ b/Engine/Plugins/Enterprise/DatasmithFBXImporter/Source/DatasmithDeltaGenTranslator/Private/DatasmithDeltaGenImporter.cpp
@@ -3,16 +3,16 @@
 #include "DatasmithDeltaGenImporter.h"
 
 #include "DatasmithDeltaGenImportData.h"
-#include "DatasmithDeltaGenImporterAuxFiles.h"
 #include "DatasmithDeltaGenImportOptions.h"
+#include "DatasmithDeltaGenImporterAuxFiles.h"
 #include "DatasmithDeltaGenLog.h"
+#include "DatasmithDeltaGenSceneProcessor.h"
 #include "DatasmithDeltaGenTranslatorModule.h"
 #include "DatasmithDeltaGenVariantConverter.h"
 #include "DatasmithFBXFileImporter.h"
-#include "DatasmithFBXImporter.h"
 #include "DatasmithFBXImportOptions.h"
+#include "DatasmithFBXImporter.h"
 #include "DatasmithFBXScene.h"
-#include "DatasmithFBXSceneProcessor.h"
 #include "DatasmithScene.h"
 #include "DatasmithSceneFactory.h"
 #include "DatasmithUtils.h"
@@ -299,7 +299,7 @@ void FDatasmithDeltaGenImporter::FetchAOTexture(const FString& MeshName, TShared
 
 bool FDatasmithDeltaGenImporter::ProcessScene()
 {
-	FDatasmithFBXSceneProcessor Processor(IntermediateScene.Get());
+	FDatasmithDeltaGenSceneProcessor Processor(IntermediateScene.Get());
 
 	// We need to create AO textures before we merge as they depend on the name of the mesh itself,
 	// and we only want to add this AO texture to the one material used by that mesh
@@ -318,22 +318,9 @@ bool FDatasmithDeltaGenImporter::ProcessScene()
 
 	Processor.FindPersistentNodes();
 
-	// TODO: BaseOptions
-	//if (Context.Options->BaseOptions.bIncludeLight)
-	//{
-		Processor.SplitLightNodes();
-	//}
-
-	// We don't export cameras from DG except for the scene camera that we
-	// ourselves create, so we don't need this.
-	//if (Context.Options->BaseOptions.bIncludeCamera)
-	//{
-		//Processor.SplitCameraNodes();
-	//}
-
-	Processor.DecomposeRotationPivots();
+	Processor.SplitLightNodes();
 
-	Processor.DecomposeScalingPivots();
+	Processor.DecomposePivots(TmlTimelines);
 
 	Processor.FindDuplicatedMaterials();
 
@@ -879,24 +866,24 @@ void PopulateTransformAnimation(IDatasmithTransformAnimationElement& TransformAn
 	EDatasmithTransformType DSType = EDatasmithTransformType::Count;
 	switch(Track.Type)
 	{
-	case FDeltaGenTmlDataAnimationTrackType::Translation:
+	case EDeltaGenTmlDataAnimationTrackType::Translation:
 	{
 		DSType = EDatasmithTransformType::Translation;
 		break;
 	}
 	// We convert Rotation to Euler on import as well
-	case FDeltaGenTmlDataAnimationTrackType::Rotation:
-	case FDeltaGenTmlDataAnimationTrackType::RotationDeltaGenEuler:
+	case EDeltaGenTmlDataAnimationTrackType::Rotation:
+	case EDeltaGenTmlDataAnimationTrackType::RotationDeltaGenEuler:
 	{
 		DSType = EDatasmithTransformType::Rotation;
 		break;
 	}
-	case FDeltaGenTmlDataAnimationTrackType::Scale:
+	case EDeltaGenTmlDataAnimationTrackType::Scale:
 	{
 		DSType = EDatasmithTransformType::Scale;
 		break;
 	}
-	case FDeltaGenTmlDataAnimationTrackType::Center:
+	case EDeltaGenTmlDataAnimationTrackType::Center:
 	{
 		UE_LOG(LogDatasmithDeltaGenImport, Warning, TEXT("Center animations are currently not supported!"));
 		return;
diff --git a/Engine/Plugins/Enterprise/DatasmithFBXImporter/Source/DatasmithDeltaGenTranslator/Private/DatasmithDeltaGenImporterAuxFiles.cpp b/Engine/Plugins/Enterprise/DatasmithFBXImporter/Source/DatasmithDeltaGenTranslator/Private/DatasmithDeltaGenImporterAuxFiles.cpp
--- a/Engine/Plugins/Enterprise/DatasmithFBXImporter/Source/DatasmithDeltaGenTranslator/Private/DatasmithDeltaGenImporterAuxFiles.cpp
+++ b/Engine/Plugins/Enterprise/DatasmithFBXImporter/Source/DatasmithDeltaGenTranslator/Private/DatasmithDeltaGenImporterAuxFiles.cpp
@@ -976,27 +976,27 @@ bool LoadTMLFile(const TCHAR* InFilePath, FDatasmithDeltaGenImportTmlResult& Out
 				}
 
 				{
-					FDeltaGenTmlDataAnimationTrackType AnimationType = FDeltaGenTmlDataAnimationTrackType::Unsupported;
+					EDeltaGenTmlDataAnimationTrackType AnimationType = EDeltaGenTmlDataAnimationTrackType::Unsupported;
 					{
 						if (TEXT("SceneObjectTranslationAnimation") == ObjectAnimationNode->GetAttribute(TEXT("Type")))
 						{
-							AnimationType = FDeltaGenTmlDataAnimationTrackType::Translation;
+							AnimationType = EDeltaGenTmlDataAnimationTrackType::Translation;
 						}
 						else if (TEXT("SceneObjectRotationAnimation") == ObjectAnimationNode->GetAttribute(TEXT("Type"))) // quaternion
 						{
-							AnimationType = FDeltaGenTmlDataAnimationTrackType::Rotation;
+							AnimationType = EDeltaGenTmlDataAnimationTrackType::Rotation;
 						}
 						else if (TEXT("SceneObjectEulerAnimation") == ObjectAnimationNode->GetAttribute(TEXT("Type")))  // degrees
 						{
-							AnimationType = FDeltaGenTmlDataAnimationTrackType::RotationDeltaGenEuler;
+							AnimationType = EDeltaGenTmlDataAnimationTrackType::RotationDeltaGenEuler;
 						}
 						else if (TEXT("SceneObjectScaleAnimation") == ObjectAnimationNode->GetAttribute(TEXT("Type")))
 						{
-							AnimationType = FDeltaGenTmlDataAnimationTrackType::Scale;
+							AnimationType = EDeltaGenTmlDataAnimationTrackType::Scale;
 						}
 						else if (TEXT("SceneObjectCenterAnimation") == ObjectAnimationNode->GetAttribute(TEXT("Type")))
 						{
-							AnimationType = FDeltaGenTmlDataAnimationTrackType::Center;
+							AnimationType = EDeltaGenTmlDataAnimationTrackType::Center;
 						}
 
 					}
@@ -1061,13 +1061,13 @@ bool LoadTMLFile(const TCHAR* InFilePath, FDatasmithDeltaGenImportTmlResult& Out
 											// E.g. 0 vs. 360 degree euler rotation is same in quaternion representation. But for animation it makes difference.
 											// For example, take animation of [(0, 0, 0), (360, 0, 0)] in euler angles, this is full circle rotation
 											// this can't be encoded in quats as quats define orientation
-											case FDeltaGenTmlDataAnimationTrackType::RotationDeltaGenEuler:
+											case EDeltaGenTmlDataAnimationTrackType::RotationDeltaGenEuler:
 											{
 												FVector Euler(ValueFloat[0], ValueFloat[1], ValueFloat[2]);
 												ValueResult = FVector4(Euler, 0.0f);
 											}
 											break;
-											case FDeltaGenTmlDataAnimationTrackType::Rotation:
+											case EDeltaGenTmlDataAnimationTrackType::Rotation:
 												{
 													// taked from FFbxDataConverter::ConvertRotToQuat
 													FQuat Quat(ValueFloat[0], -ValueFloat[1], ValueFloat[2], -ValueFloat[3]);
diff --git a/Engine/Plugins/Enterprise/DatasmithFBXImporter/Source/DatasmithDeltaGenTranslator/Private/DatasmithDeltaGenSceneProcessor.cpp b/Engine/Plugins/Enterprise/DatasmithFBXImporter/Source/DatasmithDeltaGenTranslator/Private/DatasmithDeltaGenSceneProcessor.cpp
new file mode 100644
--- /dev/null
+++ b/Engine/Plugins/Enterprise/DatasmithFBXImporter/Source/DatasmithDeltaGenTranslator/Private/DatasmithDeltaGenSceneProcessor.cpp
@@ -0,0 +1,193 @@
+// Copyright 1998-2019 Epic Games, Inc. All Rights Reserved.
+
+#include "DatasmithDeltaGenSceneProcessor.h"
+
+#include "DatasmithDeltaGenImportData.h"
+#include "DatasmithDeltaGenLog.h"
+#include "DatasmithFBXScene.h"
+
+using TimelineToAnimations = TMap<FDeltaGenTmlDataTimeline*, FDeltaGenTmlDataTimelineAnimation*>;
+
+namespace FDeltaGenProcessorImpl
+{
+	// Moves tracks of TrackTypes to new animations within each Timeline, all tied to the Dummy node's name
+	void MoveTracksToDummyAnimation(TSharedPtr<FDatasmithFBXSceneNode>& Dummy,
+									EDeltaGenTmlDataAnimationTrackType TrackTypes,
+									const FVector4& TransOffset,
+									TimelineToAnimations* FoundAnimations,
+									TMap<FDeltaGenTmlDataTimeline *, TArray<FDeltaGenTmlDataTimelineAnimation>> &NewAnimationsPerTimeline)
+	{
+		for (TPair<FDeltaGenTmlDataTimeline*, FDeltaGenTmlDataTimelineAnimation*>& AnimationInTimeline : *FoundAnimations)
+		{
+			FDeltaGenTmlDataTimeline* Timeline = AnimationInTimeline.Key;
+			FDeltaGenTmlDataTimelineAnimation* Animation = AnimationInTimeline.Value;
+
+			TArray<FDeltaGenTmlDataAnimationTrack>& Tracks = Animation->Tracks;
+			TArray<FDeltaGenTmlDataAnimationTrack> NewTracks;
+			NewTracks.Reserve(Tracks.Num());
+
+			for (int32 TrackIndex = Tracks.Num() - 1; TrackIndex >= 0; --TrackIndex)
+			{
+				FDeltaGenTmlDataAnimationTrack& ThisTrack = Tracks[TrackIndex];
+
+				if (EnumHasAnyFlags(ThisTrack.Type, TrackTypes))
+				{
+					if (EnumHasAnyFlags(ThisTrack.Type, EDeltaGenTmlDataAnimationTrackType::Translation))
+					{
+						for (FVector4& Value : ThisTrack.Values)
+						{
+							Value += TransOffset;
+						}
+					}
+
+					NewTracks.Add(ThisTrack);
+					Tracks.RemoveAt(TrackIndex);
+				}
+			}
+
+			// Move tracks to the dummy
+			if (NewTracks.Num() > 0)
+			{
+				TArray<FDeltaGenTmlDataTimelineAnimation>& NewAnimations = NewAnimationsPerTimeline.FindOrAdd(Timeline);
+
+				FDeltaGenTmlDataTimelineAnimation* NewAnimation = new(NewAnimations) FDeltaGenTmlDataTimelineAnimation;
+				NewAnimation->TargetNode = *Dummy->Name;
+				NewAnimation->Tracks = MoveTemp(NewTracks);
+
+				Dummy->MarkMovableNode();
+			}
+		}
+	}
+
+	void DecomposeRotationPivotsForNode(TSharedPtr<FDatasmithFBXSceneNode> Node,
+									    TMap<FString, TimelineToAnimations>& NodeNamesToAnimations,
+										TMap<FDeltaGenTmlDataTimeline*, TArray<FDeltaGenTmlDataTimelineAnimation>>& NewAnimationsPerTimeline)
+	{
+		if (!Node.IsValid() || Node->RotationPivot.IsNearlyZero())
+		{
+			return;
+		}
+
+		TSharedPtr<FDatasmithFBXSceneNode> NodeParent = Node->Parent.Pin();
+		if (!NodeParent.IsValid())
+		{
+			return;
+		}
+
+		FVector RotPivot = Node->RotationPivot;
+		FVector NodeLocation = Node->LocalTransform.GetTranslation();
+		FQuat NodeRotation = Node->LocalTransform.GetRotation();
+
+		Node->RotationPivot.Set(0.0f, 0.0f, 0.0f);
+		Node->LocalTransform.SetTranslation(-RotPivot);
+		Node->LocalTransform.SetRotation(FQuat::Identity);
+
+		TSharedPtr<FDatasmithFBXSceneNode> Dummy = MakeShared<FDatasmithFBXSceneNode>();
+		Dummy->Name = Node->Name + TEXT("_RotationPivot");
+		Dummy->OriginalName = Dummy->Name;
+		Dummy->SplitNodeID = Node->SplitNodeID;
+		Dummy->LocalTransform.SetTranslation(NodeLocation + RotPivot);
+		Dummy->LocalTransform.SetRotation(NodeRotation);
+
+		const FVector4 TransOffset = RotPivot + Node->RotationOffset + Node->ScalingOffset;
+
+		if (TimelineToAnimations* FoundAnimations = NodeNamesToAnimations.Find(Node->OriginalName))
+		{
+			MoveTracksToDummyAnimation(Dummy,
+									   EDeltaGenTmlDataAnimationTrackType::Rotation | EDeltaGenTmlDataAnimationTrackType::RotationDeltaGenEuler | EDeltaGenTmlDataAnimationTrackType::Translation,
+									   TransOffset,
+									   FoundAnimations,
+									   NewAnimationsPerTimeline);
+		}
+
+		// Fix hierarchy (place Dummy between Node and Parent)
+		Dummy->AddChild(Node);
+		NodeParent->Children.Remove(Node);
+		NodeParent->AddChild(Dummy);
+	}
+
+	void DecomposeScalingPivotsForNode(TSharedPtr<FDatasmithFBXSceneNode> Node,
+									   TMap<FString, TimelineToAnimations>& NodeNamesToAnimations,
+									   TMap<FDeltaGenTmlDataTimeline*, TArray<FDeltaGenTmlDataTimelineAnimation>>& NewAnimationsPerTimeline)
+	{
+		if (!Node.IsValid() || Node->ScalingPivot.IsNearlyZero())
+		{
+			return;
+		}
+
+		TSharedPtr<FDatasmithFBXSceneNode> NodeParent = Node->Parent.Pin();
+		if (!NodeParent.IsValid())
+		{
+			return;
+		}
+
+		FVector ScalingPivot = Node->ScalingPivot;
+		FVector NodeLocation = Node->LocalTransform.GetTranslation();
+		FVector NodeScaling = Node->LocalTransform.GetScale3D();
+
+		Node->ScalingPivot.Set(0.0f, 0.0f, 0.0f);
+		Node->LocalTransform.SetTranslation(-ScalingPivot);
+		Node->LocalTransform.SetScale3D(FVector::OneVector);
+
+		TSharedPtr<FDatasmithFBXSceneNode> Dummy = MakeShared<FDatasmithFBXSceneNode>();
+		Dummy->Name = Node->Name + TEXT("_ScalingPivot");
+		Dummy->OriginalName = Dummy->Name;
+		Dummy->SplitNodeID = Node->SplitNodeID;
+		Dummy->LocalTransform.SetTranslation(NodeLocation + ScalingPivot);
+		Dummy->LocalTransform.SetScale3D(NodeScaling);
+
+		const FVector4 TransOffset = ScalingPivot + Node->RotationOffset + Node->ScalingOffset;
+
+		if (TimelineToAnimations* FoundAnimations = NodeNamesToAnimations.Find(Node->OriginalName))
+		{
+			MoveTracksToDummyAnimation(Dummy,
+				EDeltaGenTmlDataAnimationTrackType::Scale | EDeltaGenTmlDataAnimationTrackType::Translation,
+				TransOffset,
+				FoundAnimations,
+				NewAnimationsPerTimeline);
+		}
+
+		// Fix hierarchy (place Dummy between Node and Parent)
+		Dummy->AddChild(Node);
+		NodeParent->Children.Remove(Node);
+		NodeParent->AddChild(Dummy);
+	}
+};
+
+FDatasmithDeltaGenSceneProcessor::FDatasmithDeltaGenSceneProcessor(FDatasmithFBXScene* InScene)
+	: FDatasmithFBXSceneProcessor(InScene)
+{
+}
+
+void FDatasmithDeltaGenSceneProcessor::DecomposePivots(TArray<FDeltaGenTmlDataTimeline>& Timelines)
+{
+	// Cache node names to all the animations they have on all timelines
+	TMap<FString, TimelineToAnimations> NodeNamesToAnimations;
+	for (FDeltaGenTmlDataTimeline& Timeline : Timelines)
+	{
+		for (FDeltaGenTmlDataTimelineAnimation& Animation : Timeline.Animations)
+		{
+			TimelineToAnimations& Animations = NodeNamesToAnimations.FindOrAdd(Animation.TargetNode.ToString());
+			Animations.Add(&Timeline, &Animation);
+		}
+	}
+
+	// Iterate over this array so that we don't step into any newly generated dummy actors
+	TMap<FDeltaGenTmlDataTimeline*, TArray<FDeltaGenTmlDataTimelineAnimation>> NewAnimationsPerTimeline;
+	for (TSharedPtr<FDatasmithFBXSceneNode> Node : Scene->GetAllNodes())
+	{
+		FDeltaGenProcessorImpl::DecomposeRotationPivotsForNode(Node, NodeNamesToAnimations, NewAnimationsPerTimeline);
+		FDeltaGenProcessorImpl::DecomposeScalingPivotsForNode(Node, NodeNamesToAnimations, NewAnimationsPerTimeline);
+	}
+
+	// Add the new animations only afterwards, as NodeNamesToAnimations stores raw pointers to animations within TArrays
+	// that might reallocate somewhere else
+	for (TTuple<FDeltaGenTmlDataTimeline*, TArray<FDeltaGenTmlDataTimelineAnimation>> Pair : NewAnimationsPerTimeline)
+	{
+		FDeltaGenTmlDataTimeline* Timeline = Pair.Key;
+		const TArray<FDeltaGenTmlDataTimelineAnimation>& NewAnimations = Pair.Value;
+
+		Timeline->Animations.Append(NewAnimations);
+	}
+}
+
diff --git a/Engine/Plugins/Enterprise/DatasmithFBXImporter/Source/DatasmithDeltaGenTranslator/Private/DatasmithDeltaGenSceneProcessor.h b/Engine/Plugins/Enterprise/DatasmithFBXImporter/Source/DatasmithDeltaGenTranslator/Private/DatasmithDeltaGenSceneProcessor.h
new file mode 100644
--- /dev/null
+++ b/Engine/Plugins/Enterprise/DatasmithFBXImporter/Source/DatasmithDeltaGenTranslator/Private/DatasmithDeltaGenSceneProcessor.h
@@ -0,0 +1,22 @@
+// Copyright 1998-2019 Epic Games, Inc. All Rights Reserved.
+
+#pragma once
+
+#include "DatasmithFBXSceneProcessor.h"
+
+#include "CoreTypes.h"
+
+struct FDatasmithFBXScene;
+struct FDeltaGenTmlDataTimeline;
+
+class FDatasmithDeltaGenSceneProcessor : public FDatasmithFBXSceneProcessor
+{
+public:
+	FDatasmithDeltaGenSceneProcessor(FDatasmithFBXScene* InScene);
+
+	/**
+	 * Decompose all scene nodes with nonzero rotation and scaling pivots
+	 * using dummy actors, and handle their animations
+	 */
+	void DecomposePivots(TArray<FDeltaGenTmlDataTimeline>& Timelines);
+};
\ No newline at end of file
diff --git a/Engine/Plugins/Enterprise/DatasmithFBXImporter/Source/DatasmithFBXTranslator/Private/DatasmithFBXSceneProcessor.cpp b/Engine/Plugins/Enterprise/DatasmithFBXImporter/Source/DatasmithFBXTranslator/Private/DatasmithFBXSceneProcessor.cpp
--- a/Engine/Plugins/Enterprise/DatasmithFBXImporter/Source/DatasmithFBXTranslator/Private/DatasmithFBXSceneProcessor.cpp
+++ b/Engine/Plugins/Enterprise/DatasmithFBXImporter/Source/DatasmithFBXTranslator/Private/DatasmithFBXSceneProcessor.cpp
@@ -406,204 +406,6 @@ void FDatasmithFBXSceneProcessor::SplitCameraNodesRecursive(TSharedPtr<FDatasmit
 	}
 }
 
-void FDatasmithFBXSceneProcessor::DecomposeRotationPivotsForNode(TSharedPtr<FDatasmithFBXSceneNode> Node, TMap<FString, FDatasmithFBXSceneAnimNode*>& NodeNameToAnimNode, TArray<FDatasmithFBXSceneAnimNode>& NewAnimNodes)
-{
-	if (Node->RotationPivot.IsNearlyZero())
-	{
-		return;
-	}
-
-	FVector RotPivot = Node->RotationPivot;
-	FVector NodeLocation = Node->LocalTransform.GetTranslation();
-	FQuat NodeRotation = Node->LocalTransform.GetRotation();
-
-	Node->RotationPivot.Set(0.0f, 0.0f, 0.0f);
-	Node->LocalTransform.SetTranslation(-RotPivot);
-	Node->LocalTransform.SetRotation(FQuat::Identity);
-
-	TSharedPtr<FDatasmithFBXSceneNode> Dummy = MakeShared<FDatasmithFBXSceneNode>();
-	Dummy->Name = Node->Name + TEXT("_RotationPivot");
-	Dummy->OriginalName = Dummy->Name;
-	Dummy->SplitNodeID = Node->SplitNodeID;
-	Dummy->LocalTransform.SetTranslation(NodeLocation + RotPivot);
-	Dummy->LocalTransform.SetRotation(NodeRotation);
-
-	// Move any rotation curves to Dummy
-	if (FDatasmithFBXSceneAnimNode** FoundNode = NodeNameToAnimNode.Find(Node->OriginalName))
-	{
-		FDatasmithFBXSceneAnimNode* NewAnimNode = nullptr;
-
-		for (FDatasmithFBXSceneAnimBlock& Block : (*FoundNode)->Blocks)
-		{
-			TArray<FDatasmithFBXSceneAnimCurve> RotCurves;
-			TArray<FDatasmithFBXSceneAnimCurve> TransCurves;
-
-			// Collect all rotation curves frm the original animnode
-			for (int32 CurveIndex = Block.Curves.Num() - 1; CurveIndex >= 0; --CurveIndex)
-			{
-				FDatasmithFBXSceneAnimCurve& ThisCurve = Block.Curves[CurveIndex];
-
-				if (ThisCurve.Type == EDatasmithFBXSceneAnimationCurveType::Rotation)
-				{
-					RotCurves.Add(ThisCurve);
-					Block.Curves.RemoveAt(CurveIndex);
-				}
-				else if (ThisCurve.Type == EDatasmithFBXSceneAnimationCurveType::Translation)
-				{
-					for (FDatasmithFBXSceneAnimPoint& Pt : ThisCurve.Points)
-					{
-						Pt.Value += (RotPivot[(uint8)ThisCurve.Component] + Node->RotationOffset[(uint8)ThisCurve.Component] + Node->ScalingOffset[(uint8)ThisCurve.Component]);
-					}
-					TransCurves.Add(ThisCurve);
-					Block.Curves.RemoveAt(CurveIndex);
-				}
-			}
-
-			// Move curves to a new block on the new animnode
-			if (RotCurves.Num() > 0 || TransCurves.Num() > 0)
-			{
-				if (NewAnimNode == nullptr)
-				{
-					NewAnimNode = new(NewAnimNodes) FDatasmithFBXSceneAnimNode;
-					NewAnimNode->Name = Dummy->Name;
-					Dummy->MarkMovableNode();
-				}
-
-				FDatasmithFBXSceneAnimBlock* NewBlock = new(NewAnimNode->Blocks) FDatasmithFBXSceneAnimBlock;
-				NewBlock->Name = Block.Name;
-				NewBlock->Curves = MoveTemp(RotCurves);
-				NewBlock->Curves.Append(TransCurves);
-			}
-		}
-	}
-
-	// Fix hierarchy (place Dummy between Node and Parent)
-	TSharedPtr<FDatasmithFBXSceneNode> NodeParent = Node->Parent.Pin();
-	Dummy->AddChild(Node);
-	if (NodeParent.IsValid())
-	{
-		NodeParent->Children.Remove(Node);
-		NodeParent->AddChild(Dummy);
-	}
-}
-
-void FDatasmithFBXSceneProcessor::DecomposeRotationPivots()
-{
-	TMap<FString, FDatasmithFBXSceneAnimNode*> NodeNameToAnimNode;
-	for (FDatasmithFBXSceneAnimNode& AnimNode : Scene->AnimNodes)
-	{
-		NodeNameToAnimNode.Add(AnimNode.Name, &AnimNode);
-	}
-
-	TArray<FDatasmithFBXSceneAnimNode> NewNodes;
-
-	for (TSharedPtr<FDatasmithFBXSceneNode> Node : Scene->GetAllNodes())
-	{
-		DecomposeRotationPivotsForNode(Node, NodeNameToAnimNode, NewNodes);
-	}
-
-	Scene->AnimNodes.Append(NewNodes);
-}
-
-void FDatasmithFBXSceneProcessor::DecomposeScalingPivotsForNode(TSharedPtr<FDatasmithFBXSceneNode> Node, TMap<FString, FDatasmithFBXSceneAnimNode*>& NodeNameToAnimNode, TArray<FDatasmithFBXSceneAnimNode>& NewAnimNodes)
-{
-	if (Node->ScalingPivot.IsNearlyZero())
-	{
-		return;
-	}
-
-	FVector ScalingPivot = Node->ScalingPivot;
-	FVector NodeLocation = Node->LocalTransform.GetTranslation();
-	FVector NodeScaling = Node->LocalTransform.GetScale3D();
-
-	Node->ScalingPivot.Set(0.0f, 0.0f, 0.0f);
-	Node->LocalTransform.SetTranslation(-ScalingPivot);
-	Node->LocalTransform.SetScale3D(FVector::OneVector);
-
-	TSharedPtr<FDatasmithFBXSceneNode> Dummy = MakeShared<FDatasmithFBXSceneNode>();
-	Dummy->Name = Node->Name + TEXT("_ScalingPivot");
-	Dummy->OriginalName = Dummy->Name;
-	Dummy->SplitNodeID = Node->SplitNodeID;
-	Dummy->LocalTransform.SetTranslation(NodeLocation + ScalingPivot);
-	Dummy->LocalTransform.SetScale3D(NodeScaling);
-
-	// Move any rotation curves to Dummy
-	if (FDatasmithFBXSceneAnimNode** FoundNode = NodeNameToAnimNode.Find(Node->OriginalName))
-	{
-		FDatasmithFBXSceneAnimNode* NewAnimNode = nullptr;
-
-		for (FDatasmithFBXSceneAnimBlock& Block : (*FoundNode)->Blocks)
-		{
-			TArray<FDatasmithFBXSceneAnimCurve> ScaleCurves;
-			TArray<FDatasmithFBXSceneAnimCurve> TransCurves;
-
-			// Collect all rotation curves frm the original animnode
-			for (int32 CurveIndex = Block.Curves.Num() - 1; CurveIndex >= 0; --CurveIndex)
-			{
-				FDatasmithFBXSceneAnimCurve& ThisCurve = Block.Curves[CurveIndex];
-
-				if (ThisCurve.Type == EDatasmithFBXSceneAnimationCurveType::Scale)
-				{
-					ScaleCurves.Add(ThisCurve);
-					Block.Curves.RemoveAt(CurveIndex);
-				}
-				else if (ThisCurve.Type == EDatasmithFBXSceneAnimationCurveType::Translation)
-				{
-					for (FDatasmithFBXSceneAnimPoint& Pt : ThisCurve.Points)
-					{
-						Pt.Value += (ScalingPivot[(uint8)ThisCurve.Component] + Node->RotationOffset[(uint8)ThisCurve.Component] + Node->ScalingOffset[(uint8)ThisCurve.Component]);
-					}
-					TransCurves.Add(ThisCurve);
-					Block.Curves.RemoveAt(CurveIndex);
-				}
-			}
-
-			// Move curves to a new block on the new animnode
-			if (ScaleCurves.Num() > 0 || TransCurves.Num() > 0)
-			{
-				if (NewAnimNode == nullptr)
-				{
-					NewAnimNode = new(NewAnimNodes) FDatasmithFBXSceneAnimNode;
-					NewAnimNode->Name = Dummy->Name;
-					Dummy->MarkMovableNode();
-				}
-
-				FDatasmithFBXSceneAnimBlock* NewBlock = new(NewAnimNode->Blocks) FDatasmithFBXSceneAnimBlock;
-				NewBlock->Name = Block.Name;
-				NewBlock->Curves = MoveTemp(ScaleCurves);
-				NewBlock->Curves.Append(TransCurves);
-			}
-		}
-	}
-
-	// Fix hierarchy (place Dummy between Node and Parent)
-	TSharedPtr<FDatasmithFBXSceneNode> NodeParent = Node->Parent.Pin();
-	Dummy->AddChild(Node);
-	if (NodeParent.IsValid())
-	{
-		NodeParent->Children.Remove(Node);
-		NodeParent->AddChild(Dummy);
-	}
-}
-
-void FDatasmithFBXSceneProcessor::DecomposeScalingPivots()
-{
-	TMap<FString, FDatasmithFBXSceneAnimNode*> NodeNameToAnimNode;
-	for (FDatasmithFBXSceneAnimNode& AnimNode : Scene->AnimNodes)
-	{
-		NodeNameToAnimNode.Add(AnimNode.Name, &AnimNode);
-	}
-
-	TArray<FDatasmithFBXSceneAnimNode> NewNodes;
-
-	for (TSharedPtr<FDatasmithFBXSceneNode> Node : Scene->GetAllNodes())
-	{
-		DecomposeScalingPivotsForNode(Node, NodeNameToAnimNode, NewNodes);
-	}
-
-	Scene->AnimNodes.Append(NewNodes);
-}
-
 struct FSwitchState
 {
 	FSwitchState(TSharedPtr<FDatasmithFBXSceneNode> Node)
diff --git a/Engine/Plugins/Enterprise/DatasmithFBXImporter/Source/DatasmithFBXTranslator/Public/DatasmithFBXSceneProcessor.h b/Engine/Plugins/Enterprise/DatasmithFBXImporter/Source/DatasmithFBXTranslator/Public/DatasmithFBXSceneProcessor.h
--- a/Engine/Plugins/Enterprise/DatasmithFBXImporter/Source/DatasmithFBXTranslator/Public/DatasmithFBXSceneProcessor.h
+++ b/Engine/Plugins/Enterprise/DatasmithFBXImporter/Source/DatasmithFBXTranslator/Public/DatasmithFBXSceneProcessor.h
@@ -78,18 +78,6 @@ public:
 	/** Recursively splits cameras into a parent simple actor, and a child camera actor, with the _Camera suffix */
 	void SplitCameraNodesRecursive(TSharedPtr<FDatasmithFBXSceneNode> Node);
 
-	/** Decompose all scene nodes with nonzero RotationPivots using dummy actors, and handle their animations */
-	void DecomposeRotationPivots();
-
-	/** Recursively decompose nodes with nonzero RotationPivots using dummy actors, and handle their animations */
-	void DecomposeRotationPivotsForNode(TSharedPtr<FDatasmithFBXSceneNode> Node, TMap<FString, FDatasmithFBXSceneAnimNode*>& NodeNameToAnimNode, TArray<FDatasmithFBXSceneAnimNode>& NewAnimNodes);
-
-	/** Decompose all scene nodes with nonzero ScalingPivots using dummy actors, and handle their animations */
-	void DecomposeScalingPivots();
-
-	/** Recursively decompose nodes with nonzero ScalingPivots using dummy actors, and handle their animations */
-	void DecomposeScalingPivotsForNode(TSharedPtr<FDatasmithFBXSceneNode> Node, TMap<FString, FDatasmithFBXSceneAnimNode*>& NodeNameToAnimNode, TArray<FDatasmithFBXSceneAnimNode>& NewAnimNodes);
-
 protected:
 	FDatasmithFBXScene* Scene;
 };
diff --git a/Engine/Plugins/Enterprise/DatasmithFBXImporter/Source/DatasmithVREDTranslator/Private/DatasmithVREDSceneProcessor.cpp b/Engine/Plugins/Enterprise/DatasmithFBXImporter/Source/DatasmithVREDTranslator/Private/DatasmithVREDSceneProcessor.cpp
--- a/Engine/Plugins/Enterprise/DatasmithFBXImporter/Source/DatasmithVREDTranslator/Private/DatasmithVREDSceneProcessor.cpp
+++ b/Engine/Plugins/Enterprise/DatasmithFBXImporter/Source/DatasmithVREDTranslator/Private/DatasmithVREDSceneProcessor.cpp
@@ -62,4 +62,206 @@ void FDatasmithVREDSceneProcessor::AddMatsMaterials(TArray<FDatasmithFBXSceneMat
 			Scene->Materials.Add(AddedMat);
 		}
 	}
+}
+
+void FDatasmithVREDSceneProcessor::DecomposeRotationPivotsForNode(TSharedPtr<FDatasmithFBXSceneNode> Node, TMap<FString, FDatasmithFBXSceneAnimNode*>& NodeNameToAnimNode, TArray<FDatasmithFBXSceneAnimNode>& NewAnimNodes)
+{
+	if (!Node.IsValid() || Node->RotationPivot.IsNearlyZero())
+	{
+		return;
+	}
+
+	TSharedPtr<FDatasmithFBXSceneNode> NodeParent = Node->Parent.Pin();
+	if (!NodeParent.IsValid())
+	{
+		return;
+	}
+
+	FVector RotPivot = Node->RotationPivot;
+	FVector NodeLocation = Node->LocalTransform.GetTranslation();
+	FQuat NodeRotation = Node->LocalTransform.GetRotation();
+
+	Node->RotationPivot.Set(0.0f, 0.0f, 0.0f);
+	Node->LocalTransform.SetTranslation(-RotPivot);
+	Node->LocalTransform.SetRotation(FQuat::Identity);
+
+	TSharedPtr<FDatasmithFBXSceneNode> Dummy = MakeShared<FDatasmithFBXSceneNode>();
+	Dummy->Name = Node->Name + TEXT("_RotationPivot");
+	Dummy->OriginalName = Dummy->Name;
+	Dummy->SplitNodeID = Node->SplitNodeID;
+	Dummy->LocalTransform.SetTranslation(NodeLocation + RotPivot);
+	Dummy->LocalTransform.SetRotation(NodeRotation);
+
+	// Move any rotation curves to Dummy
+	if (FDatasmithFBXSceneAnimNode** FoundNode = NodeNameToAnimNode.Find(Node->OriginalName))
+	{
+		FDatasmithFBXSceneAnimNode* NewAnimNode = nullptr;
+
+		for (FDatasmithFBXSceneAnimBlock& Block : (*FoundNode)->Blocks)
+		{
+			TArray<FDatasmithFBXSceneAnimCurve> RotCurves;
+			TArray<FDatasmithFBXSceneAnimCurve> TransCurves;
+
+			// Collect all rotation curves frm the original animnode
+			for (int32 CurveIndex = Block.Curves.Num() - 1; CurveIndex >= 0; --CurveIndex)
+			{
+				FDatasmithFBXSceneAnimCurve& ThisCurve = Block.Curves[CurveIndex];
+
+				if (ThisCurve.Type == EDatasmithFBXSceneAnimationCurveType::Rotation)
+				{
+					RotCurves.Add(ThisCurve);
+					Block.Curves.RemoveAt(CurveIndex);
+				}
+				else if (ThisCurve.Type == EDatasmithFBXSceneAnimationCurveType::Translation)
+				{
+					for (FDatasmithFBXSceneAnimPoint& Pt : ThisCurve.Points)
+					{
+						Pt.Value += (RotPivot[(uint8)ThisCurve.Component] + Node->RotationOffset[(uint8)ThisCurve.Component] + Node->ScalingOffset[(uint8)ThisCurve.Component]);
+					}
+					TransCurves.Add(ThisCurve);
+					Block.Curves.RemoveAt(CurveIndex);
+				}
+			}
+
+			// Move curves to a new block on the new animnode
+			if (RotCurves.Num() > 0 || TransCurves.Num() > 0)
+			{
+				if (NewAnimNode == nullptr)
+				{
+					NewAnimNode = new(NewAnimNodes) FDatasmithFBXSceneAnimNode;
+					NewAnimNode->Name = Dummy->Name;
+					Dummy->MarkMovableNode();
+				}
+
+				FDatasmithFBXSceneAnimBlock* NewBlock = new(NewAnimNode->Blocks) FDatasmithFBXSceneAnimBlock;
+				NewBlock->Name = Block.Name;
+				NewBlock->Curves = MoveTemp(RotCurves);
+				NewBlock->Curves.Append(TransCurves);
+			}
+		}
+	}
+
+	// Fix hierarchy (place Dummy between Node and Parent)
+	Dummy->AddChild(Node);
+	NodeParent->Children.Remove(Node);
+	NodeParent->AddChild(Dummy);
+}
+
+void FDatasmithVREDSceneProcessor::DecomposeRotationPivots()
+{
+	TMap<FString, FDatasmithFBXSceneAnimNode*> NodeNameToAnimNode;
+	for (FDatasmithFBXSceneAnimNode& AnimNode : Scene->AnimNodes)
+	{
+		NodeNameToAnimNode.Add(AnimNode.Name, &AnimNode);
+	}
+
+	TArray<FDatasmithFBXSceneAnimNode> NewNodes;
+
+	for (TSharedPtr<FDatasmithFBXSceneNode> Node : Scene->GetAllNodes())
+	{
+		DecomposeRotationPivotsForNode(Node, NodeNameToAnimNode, NewNodes);
+	}
+
+	Scene->AnimNodes.Append(NewNodes);
+}
+
+void FDatasmithVREDSceneProcessor::DecomposeScalingPivotsForNode(TSharedPtr<FDatasmithFBXSceneNode> Node, TMap<FString, FDatasmithFBXSceneAnimNode*>& NodeNameToAnimNode, TArray<FDatasmithFBXSceneAnimNode>& NewAnimNodes)
+{
+	if (!Node.IsValid() || Node->ScalingPivot.IsNearlyZero())
+	{
+		return;
+	}
+
+	TSharedPtr<FDatasmithFBXSceneNode> NodeParent = Node->Parent.Pin();
+	if (!NodeParent.IsValid())
+	{
+		return;
+	}
+
+	FVector ScalingPivot = Node->ScalingPivot;
+	FVector NodeLocation = Node->LocalTransform.GetTranslation();
+	FVector NodeScaling = Node->LocalTransform.GetScale3D();
+
+	Node->ScalingPivot.Set(0.0f, 0.0f, 0.0f);
+	Node->LocalTransform.SetTranslation(-ScalingPivot);
+	Node->LocalTransform.SetScale3D(FVector::OneVector);
+
+	TSharedPtr<FDatasmithFBXSceneNode> Dummy = MakeShared<FDatasmithFBXSceneNode>();
+	Dummy->Name = Node->Name + TEXT("_ScalingPivot");
+	Dummy->OriginalName = Dummy->Name;
+	Dummy->SplitNodeID = Node->SplitNodeID;
+	Dummy->LocalTransform.SetTranslation(NodeLocation + ScalingPivot);
+	Dummy->LocalTransform.SetScale3D(NodeScaling);
+
+	// Move any rotation curves to Dummy
+	if (FDatasmithFBXSceneAnimNode** FoundNode = NodeNameToAnimNode.Find(Node->OriginalName))
+	{
+		FDatasmithFBXSceneAnimNode* NewAnimNode = nullptr;
+
+		for (FDatasmithFBXSceneAnimBlock& Block : (*FoundNode)->Blocks)
+		{
+			TArray<FDatasmithFBXSceneAnimCurve> ScaleCurves;
+			TArray<FDatasmithFBXSceneAnimCurve> TransCurves;
+
+			// Collect all rotation curves frm the original animnode
+			for (int32 CurveIndex = Block.Curves.Num() - 1; CurveIndex >= 0; --CurveIndex)
+			{
+				FDatasmithFBXSceneAnimCurve& ThisCurve = Block.Curves[CurveIndex];
+
+				if (ThisCurve.Type == EDatasmithFBXSceneAnimationCurveType::Scale)
+				{
+					ScaleCurves.Add(ThisCurve);
+					Block.Curves.RemoveAt(CurveIndex);
+				}
+				else if (ThisCurve.Type == EDatasmithFBXSceneAnimationCurveType::Translation)
+				{
+					for (FDatasmithFBXSceneAnimPoint& Pt : ThisCurve.Points)
+					{
+						Pt.Value += (ScalingPivot[(uint8)ThisCurve.Component] + Node->RotationOffset[(uint8)ThisCurve.Component] + Node->ScalingOffset[(uint8)ThisCurve.Component]);
+					}
+					TransCurves.Add(ThisCurve);
+					Block.Curves.RemoveAt(CurveIndex);
+				}
+			}
+
+			// Move curves to a new block on the new animnode
+			if (ScaleCurves.Num() > 0 || TransCurves.Num() > 0)
+			{
+				if (NewAnimNode == nullptr)
+				{
+					NewAnimNode = new(NewAnimNodes) FDatasmithFBXSceneAnimNode;
+					NewAnimNode->Name = Dummy->Name;
+					Dummy->MarkMovableNode();
+				}
+
+				FDatasmithFBXSceneAnimBlock* NewBlock = new(NewAnimNode->Blocks) FDatasmithFBXSceneAnimBlock;
+				NewBlock->Name = Block.Name;
+				NewBlock->Curves = MoveTemp(ScaleCurves);
+				NewBlock->Curves.Append(TransCurves);
+			}
+		}
+	}
+
+	// Fix hierarchy (place Dummy between Node and Parent)
+	Dummy->AddChild(Node);
+	NodeParent->Children.Remove(Node);
+	NodeParent->AddChild(Dummy);
+}
+
+void FDatasmithVREDSceneProcessor::DecomposeScalingPivots()
+{
+	TMap<FString, FDatasmithFBXSceneAnimNode*> NodeNameToAnimNode;
+	for (FDatasmithFBXSceneAnimNode& AnimNode : Scene->AnimNodes)
+	{
+		NodeNameToAnimNode.Add(AnimNode.Name, &AnimNode);
+	}
+
+	TArray<FDatasmithFBXSceneAnimNode> NewNodes;
+
+	for (TSharedPtr<FDatasmithFBXSceneNode> Node : Scene->GetAllNodes())
+	{
+		DecomposeScalingPivotsForNode(Node, NodeNameToAnimNode, NewNodes);
+	}
+
+	Scene->AnimNodes.Append(NewNodes);
 }
\ No newline at end of file
diff --git a/Engine/Plugins/Enterprise/DatasmithFBXImporter/Source/DatasmithVREDTranslator/Private/DatasmithVREDSceneProcessor.h b/Engine/Plugins/Enterprise/DatasmithFBXImporter/Source/DatasmithVREDTranslator/Private/DatasmithVREDSceneProcessor.h
--- a/Engine/Plugins/Enterprise/DatasmithFBXImporter/Source/DatasmithVREDTranslator/Private/DatasmithVREDSceneProcessor.h
+++ b/Engine/Plugins/Enterprise/DatasmithFBXImporter/Source/DatasmithVREDTranslator/Private/DatasmithVREDSceneProcessor.h
@@ -24,6 +24,18 @@ public:
 	/** Overwrite FBX imported materials with mats material parameters and info */
 	void AddMatsMaterials(TArray<FDatasmithFBXSceneMaterial>* InMatsMaterials);
 
+	/** Decompose all scene nodes with nonzero RotationPivots using dummy actors, and handle their animations */
+	void DecomposeRotationPivots();
+
+	/** Recursively decompose nodes with nonzero RotationPivots using dummy actors, and handle their animations */
+	void DecomposeRotationPivotsForNode(TSharedPtr<FDatasmithFBXSceneNode> Node, TMap<FString, FDatasmithFBXSceneAnimNode*>& NodeNameToAnimNode, TArray<FDatasmithFBXSceneAnimNode>& NewAnimNodes);
+
+	/** Decompose all scene nodes with nonzero ScalingPivots using dummy actors, and handle their animations */
+	void DecomposeScalingPivots();
+
+	/** Recursively decompose nodes with nonzero ScalingPivots using dummy actors, and handle their animations */
+	void DecomposeScalingPivotsForNode(TSharedPtr<FDatasmithFBXSceneNode> Node, TMap<FString, FDatasmithFBXSceneAnimNode*>& NodeNameToAnimNode, TArray<FDatasmithFBXSceneAnimNode>& NewAnimNodes);
+
 protected:
 	TMap<FString, TSharedPtr<FDatasmithFBXSceneLight>> ExtraLightsInfo;
 };
\ No newline at end of file
