Moved the VRED Extras to Resources so that they are part of the binary package.

#rnx

diff --git a/Engine/Plugins/Enterprise/DatasmithFBXImporter/Extras/PythonAPI/test_deltagen_api.py b/Engine/Plugins/Enterprise/DatasmithFBXImporter/Extras/PythonAPI/test_deltagen_api.py
deleted file mode 100644
diff --git a/Engine/Plugins/Enterprise/DatasmithFBXImporter/Extras/PythonAPI/test_vred_api.py b/Engine/Plugins/Enterprise/DatasmithFBXImporter/Extras/PythonAPI/test_vred_api.py
deleted file mode 100644
diff --git a/Engine/Plugins/Enterprise/DatasmithFBXImporter/Extras/VREDPlugin/DatasmithExporter/__init__.py b/Engine/Plugins/Enterprise/DatasmithFBXImporter/Extras/VREDPlugin/DatasmithExporter/__init__.py
deleted file mode 100644
diff --git a/Engine/Plugins/Enterprise/DatasmithFBXImporter/Extras/VREDPlugin/DatasmithExporter/anim_exporter.py b/Engine/Plugins/Enterprise/DatasmithFBXImporter/Extras/VREDPlugin/DatasmithExporter/anim_exporter.py
deleted file mode 100644
diff --git a/Engine/Plugins/Enterprise/DatasmithFBXImporter/Extras/VREDPlugin/DatasmithExporter/exporter.py b/Engine/Plugins/Enterprise/DatasmithFBXImporter/Extras/VREDPlugin/DatasmithExporter/exporter.py
deleted file mode 100644
diff --git a/Engine/Plugins/Enterprise/DatasmithFBXImporter/Extras/VREDPlugin/DatasmithExporter/log.py b/Engine/Plugins/Enterprise/DatasmithFBXImporter/Extras/VREDPlugin/DatasmithExporter/log.py
deleted file mode 100644
diff --git a/Engine/Plugins/Enterprise/DatasmithFBXImporter/Extras/VREDPlugin/DatasmithExporter/main.py b/Engine/Plugins/Enterprise/DatasmithFBXImporter/Extras/VREDPlugin/DatasmithExporter/main.py
deleted file mode 100644
diff --git a/Engine/Plugins/Enterprise/DatasmithFBXImporter/Extras/VREDPlugin/DatasmithExporter/mat.py b/Engine/Plugins/Enterprise/DatasmithFBXImporter/Extras/VREDPlugin/DatasmithExporter/mat.py
deleted file mode 100644
diff --git a/Engine/Plugins/Enterprise/DatasmithFBXImporter/Extras/VREDPlugin/DatasmithExporter/sanitizer.py b/Engine/Plugins/Enterprise/DatasmithFBXImporter/Extras/VREDPlugin/DatasmithExporter/sanitizer.py
deleted file mode 100644
diff --git a/Engine/Plugins/Enterprise/DatasmithFBXImporter/Extras/VREDPlugin/DatasmithExporter/utils.py b/Engine/Plugins/Enterprise/DatasmithFBXImporter/Extras/VREDPlugin/DatasmithExporter/utils.py
deleted file mode 100644
diff --git a/Engine/Plugins/Enterprise/DatasmithFBXImporter/Resources/PythonAPI/test_deltagen_api.py b/Engine/Plugins/Enterprise/DatasmithFBXImporter/Resources/PythonAPI/test_deltagen_api.py
new file mode 100644
--- /dev/null
+++ b/Engine/Plugins/Enterprise/DatasmithFBXImporter/Resources/PythonAPI/test_deltagen_api.py
@@ -0,0 +1,96 @@
+import unreal
+
+file_a = "C:\\MyScenes\\fulltest.fbx"
+file_b = "C:\\MyScenes\\fulltest.fbx"
+imported_scenes_path = "/Game/ImportedScenes"
+
+print 'Preparing import options...'
+advanced_mesh_options = unreal.DatasmithStaticMeshImportOptions()
+advanced_mesh_options.set_editor_property('max_lightmap_resolution', unreal.DatasmithImportLightmapMax.LIGHTMAP_512)
+advanced_mesh_options.set_editor_property('min_lightmap_resolution', unreal.DatasmithImportLightmapMin.LIGHTMAP_64)
+advanced_mesh_options.set_editor_property('generate_lightmap_u_vs', True)
+advanced_mesh_options.set_editor_property('remove_degenerates', True)
+
+base_options = unreal.DatasmithImportBaseOptions()
+base_options.set_editor_property('include_geometry', True)
+base_options.set_editor_property('include_material', True)
+base_options.set_editor_property('include_light', True)
+base_options.set_editor_property('include_camera', True)
+base_options.set_editor_property('include_animation', True)
+base_options.set_editor_property('static_mesh_options', advanced_mesh_options)
+base_options.set_editor_property('scene_handling', unreal.DatasmithImportScene.CURRENT_LEVEL)
+base_options.set_editor_property('asset_options', [])  # Not used
+
+dg_options = unreal.DatasmithDeltaGenImportOptions()
+dg_options.set_editor_property('merge_nodes', False)
+dg_options.set_editor_property('optimize_duplicated_nodes', False)
+dg_options.set_editor_property('remove_invisible_nodes', False)
+dg_options.set_editor_property('simplify_node_hierarchy', False)
+dg_options.set_editor_property('import_var', True)
+dg_options.set_editor_property('var_path', "")
+dg_options.set_editor_property('import_pos', True)
+dg_options.set_editor_property('pos_path', "")
+dg_options.set_editor_property('import_tml', True)
+dg_options.set_editor_property('tml_path', "")
+dg_options.set_editor_property('textures_dir', "")
+dg_options.set_editor_property('intermediate_serialization', unreal.DatasmithDeltaGenIntermediateSerializationType.DISABLED)
+dg_options.set_editor_property('colorize_materials', False)
+dg_options.set_editor_property('generate_lightmap_u_vs', False)
+dg_options.set_editor_property('import_animations', True)
+
+# Direct import to scene and assets:
+print 'Importing directly to scene...'
+unreal.DeltaGenLibrary.import_(file_a, imported_scenes_path, base_options, None, False)
+
+#2-stage import step 1:
+print 'Parsing to scene object...'
+scene = unreal.DatasmithDeltaGenSceneElement.construct_datasmith_scene_from_file(file_b, imported_scenes_path, base_options, dg_options)
+print 'Resulting datasmith scene: ' + str(scene)
+print '\tProduct name: ' + str(scene.get_product_name())
+print '\tMesh actor count: ' + str(len(scene.get_all_mesh_actors()))
+print '\tLight actor count: ' + str(len(scene.get_all_light_actors()))
+print '\tCamera actor count: ' + str(len(scene.get_all_camera_actors()))
+print '\tCustom actor count: ' + str(len(scene.get_all_custom_actors()))
+print '\tMaterial count: ' + str(len(scene.get_all_materials()))
+print '\tAnimationTimeline count: ' + str(len(scene.get_all_animation_timelines()))
+print '\tVariant count: ' + str(len(scene.get_all_variants()))
+
+# Modify one of the Timelines
+# Warning: The Animation nested structure is all USTRUCTs, which are value types, and the Array accessor returns 
+# a copy. Meaning something like timeline[0].name = 'new_name' will set the name on the COPY of anim_nodes[0]
+timelines = scene.get_all_animation_timelines()
+if len(timelines) > 0:
+    tim_0 = timelines[0]
+    old_name = tim_0.name 
+    print 'Timeline old name: ' + old_name
+    
+    tim_0.name += '_MODIFIED'
+    modified_name = tim_0.name    
+    print 'Anim node modified name: ' + modified_name
+    
+    timelines[0] = tim_0
+    scene.set_all_animation_timelines(timelines)
+    
+    # Check modification
+    new_timelines = scene.get_all_animation_timelines()
+    print 'Anim node retrieved modified name: ' + new_timelines[0].name
+    assert new_timelines[0].name == modified_name, "Node modification didn't work!"
+
+    # Restore to previous state
+    tim_0 = new_timelines[0]
+    tim_0.name = old_name
+    new_timelines[0] = tim_0
+    scene.set_all_animation_timelines(new_timelines)
+    
+# 2-stage import step 2:
+print 'Importing assets and actors...'
+result = scene.import_scene()
+print 'Import results: '
+print '\tImported actor count: ' + str(len(result.imported_actors))
+print '\tImported mesh count: ' + str(len(result.imported_meshes))
+print '\tImported level sequences: ' + str([a.get_name() for a in result.animations])
+print '\tImported level variant sets asset: ' + str(result.level_variant_sets.get_name())
+if result.import_succeed:
+    print 'Import succeeded!'
+else:
+    print 'Import failed!'
diff --git a/Engine/Plugins/Enterprise/DatasmithFBXImporter/Resources/PythonAPI/test_vred_api.py b/Engine/Plugins/Enterprise/DatasmithFBXImporter/Resources/PythonAPI/test_vred_api.py
new file mode 100644
--- /dev/null
+++ b/Engine/Plugins/Enterprise/DatasmithFBXImporter/Resources/PythonAPI/test_vred_api.py
@@ -0,0 +1,97 @@
+import unreal
+
+file_a = "C:\\MyScenes\\file_a.fbx"
+file_b = "C:\\MyScenes\\file_a.fbx"
+imported_scenes_path = "/Game/ImportedScenes"
+
+print 'Preparing import options...'
+advanced_mesh_options = unreal.DatasmithStaticMeshImportOptions()
+advanced_mesh_options.set_editor_property('max_lightmap_resolution', unreal.DatasmithImportLightmapMax.LIGHTMAP_512)
+advanced_mesh_options.set_editor_property('min_lightmap_resolution', unreal.DatasmithImportLightmapMin.LIGHTMAP_64)
+advanced_mesh_options.set_editor_property('generate_lightmap_u_vs', True)
+advanced_mesh_options.set_editor_property('remove_degenerates', True)
+
+base_options = unreal.DatasmithImportBaseOptions()
+base_options.set_editor_property('include_geometry', True)
+base_options.set_editor_property('include_material', True)
+base_options.set_editor_property('include_light', True)
+base_options.set_editor_property('include_camera', True)
+base_options.set_editor_property('include_animation', True)
+base_options.set_editor_property('static_mesh_options', advanced_mesh_options)
+base_options.set_editor_property('scene_handling', unreal.DatasmithImportScene.CURRENT_LEVEL)
+base_options.set_editor_property('asset_options', [])  # Not used
+
+vred_options = unreal.DatasmithVREDImportOptions()
+vred_options.set_editor_property('merge_nodes', False)
+vred_options.set_editor_property('optimize_duplicated_nodes', False)
+vred_options.set_editor_property('import_var', True)
+vred_options.set_editor_property('var_path', "")
+vred_options.set_editor_property('import_light_info', True)
+vred_options.set_editor_property('light_info_path', "")
+vred_options.set_editor_property('import_clip_info', True)
+vred_options.set_editor_property('clip_info_path', "")
+vred_options.set_editor_property('textures_dir', "")
+vred_options.set_editor_property('import_animations', True)
+vred_options.set_editor_property('intermediate_serialization', unreal.DatasmithVREDIntermediateSerializationType.DISABLED)
+vred_options.set_editor_property('colorize_materials', False)
+vred_options.set_editor_property('generate_lightmap_u_vs', False)
+vred_options.set_editor_property('import_animations', True)
+
+# Direct import to scene and assets:
+print 'Importing directly to scene...'
+unreal.VREDLibrary.import_(file_a, imported_scenes_path, base_options, None, True)
+
+#2-stage import step 1:
+print 'Parsing to scene object...'
+scene = unreal.DatasmithVREDSceneElement.construct_datasmith_scene_from_file(file_b, imported_scenes_path, base_options, vred_options)
+print 'Resulting datasmith scene: ' + str(scene)
+print '\tProduct name: ' + str(scene.get_product_name())
+print '\tMesh actor count: ' + str(len(scene.get_all_mesh_actors()))
+print '\tLight actor count: ' + str(len(scene.get_all_light_actors()))
+print '\tCamera actor count: ' + str(len(scene.get_all_camera_actors()))
+print '\tCustom actor count: ' + str(len(scene.get_all_custom_actors()))
+print '\tMaterial count: ' + str(len(scene.get_all_materials()))
+print '\tAnimNode count: ' + str(len(scene.get_all_anim_nodes()))
+print '\tAnimClip count: ' + str(len(scene.get_all_anim_clips()))
+print '\tExtra light info count: ' + str(len(scene.get_all_extra_lights_info()))
+print '\tVariant count: ' + str(len(scene.get_all_variants()))
+
+# Modify one of the AnimNodes
+# Warning: The AnimNode nested structure is all USTRUCTs, which are value types, and the Array accessor returns 
+# a copy. Meaning something like anim_nodes[0].name = 'new_name' will set the name on the COPY of anim_nodes[0]
+anim_nodes = scene.get_all_anim_nodes()
+if len(anim_nodes) > 0:
+    node_0 = anim_nodes[0]
+    old_name = node_0.name 
+    print 'Anim node old name: ' + old_name
+    
+    node_0.name += '_MODIFIED'
+    modified_name = node_0.name    
+    print 'Anim node modified name: ' + modified_name
+    
+    anim_nodes[0] = node_0
+    scene.set_all_anim_nodes(anim_nodes)
+    
+    # Check modification
+    new_anim_nodes = scene.get_all_anim_nodes()
+    print 'Anim node retrieved modified name: ' + new_anim_nodes[0].name
+    assert new_anim_nodes[0].name == modified_name, "Node modification didn't work!"
+
+    # Restore to previous state
+    node_0 = new_anim_nodes[0]
+    node_0.name = old_name
+    new_anim_nodes[0] = node_0
+    scene.set_all_anim_nodes(new_anim_nodes)
+    
+# 2-stage import step 2:
+print 'Importing assets and actors...'
+result = scene.import_scene()
+print 'Import results: '
+print '\tImported actor count: ' + str(len(result.imported_actors))
+print '\tImported mesh count: ' + str(len(result.imported_meshes))
+print '\tImported level sequences: ' + str([a.get_name() for a in result.animations])
+print '\tImported level variant sets asset: ' + str(result.level_variant_sets.get_name())
+if result.import_succeed:
+    print 'Import succeeded!'
+else:
+    print 'Import failed!'
diff --git a/Engine/Plugins/Enterprise/DatasmithFBXImporter/Resources/VREDPlugin/DatasmithExporter/__init__.py b/Engine/Plugins/Enterprise/DatasmithFBXImporter/Resources/VREDPlugin/DatasmithExporter/__init__.py
new file mode 100644
--- /dev/null
+++ b/Engine/Plugins/Enterprise/DatasmithFBXImporter/Resources/VREDPlugin/DatasmithExporter/__init__.py
@@ -0,0 +1,28 @@
+from PySide2 import QtCore, QtGui, QtWidgets
+
+# This has to be run after everything is imported, so that we can remove all entries
+# that our scripts created. For some reason VRED likes parsing this file last, so
+# we placed the menu cleaning code in here
+
+def vredMainWindow(id):
+    from shiboken2 import wrapInstance
+    return wrapInstance(id, QtWidgets.QMainWindow)
+
+# Only run this when VRED is parsing this file
+if 'VREDModule' in __name__ or 'PythonQt_module' in __name__:
+    # Strings to remove from the Scripts menu. Could do something with the os module but
+    # this is safer, and shouldn't change that often
+    fileNames = ['anim_exporter', 'log', 'main', 'sanitizer', 'utils', 'exporter', 'mat', '__init__']
+
+    mainWindow = vredMainWindow(VREDMainWindowId)
+    mb = mainWindow.menuBar()
+    actions = mb.actions()
+
+    # Iterate over the main menu bar
+    for action in actions:
+        # Iterate over the Scripts menu and remove entries in fileNames
+        if action.text() == mainWindow.tr("Scripts"):
+            scriptsMenuItems = action.menu().actions()
+            for scriptsItem in scriptsMenuItems:
+                if scriptsItem.text() in fileNames:
+                    scriptsItem.setVisible(False)
diff --git a/Engine/Plugins/Enterprise/DatasmithFBXImporter/Resources/VREDPlugin/DatasmithExporter/anim_exporter.py b/Engine/Plugins/Enterprise/DatasmithFBXImporter/Resources/VREDPlugin/DatasmithExporter/anim_exporter.py
new file mode 100644
--- /dev/null
+++ b/Engine/Plugins/Enterprise/DatasmithFBXImporter/Resources/VREDPlugin/DatasmithExporter/anim_exporter.py
@@ -0,0 +1,612 @@
+import sys
+import os
+import logging
+
+import vrScenegraph
+import vrAnimWidgets
+import vrFieldAccess
+import vrNodePtr
+import vrOSGTypes
+import vrNodeUtils
+import vrController
+import vrOptimize
+
+# Adds the ScriptPlugins folder to path so that we can find the other files
+# in our module
+scriptPluginsFolder = os.path.dirname(os.path.realpath(os.curdir))
+if scriptPluginsFolder not in sys.path:
+    sys.path.append(scriptPluginsFolder)
+
+import DatasmithExporter.utils as dsutils
+import DatasmithExporter.sanitizer as dss
+
+# Time where we place the tag key frame to mark the curves
+TIME_OF_TAG_KEY = 100
+
+class DSTransformStorage:
+    """
+    We use these objects to store the transform of vrNodePtr before we call resetAnim(clipname) on
+    the clips, since there is no intended way of returning back to the "unanimated" transform
+    """
+    def __init__(self):
+        self.Translation = []
+        self.Rotation = []
+        self.Scale = []
+        self.RotationOrientation = []
+        self.Visible = True
+
+class DSAnimNode:
+    def __init__(self, name):
+        self.Name = name
+        self.AnimBlocks = []
+
+class DSAnimBlock:
+    def __init__(self, id, name):
+        self.DSID = id
+        self.Name = name
+        self.ChildCurveNamesToCurves = {}
+
+    # Define these so that we can use the same sanitizer functions that
+    # we use for VRED scene nodes
+    def getName(self):
+        return self.Name
+    def setName(self, value):
+        self.Name = value
+    def getID(self):
+        return id(self)
+
+class DSAnimCurve:
+    def __init__(self, id, animCurve, startSeconds):
+        self.DSID = id
+        self.AnimCurve = animCurve
+        self.StartSeconds = startSeconds
+
+class DSAnimNodeUsage:
+    """
+    Represents the usage of an AnimNode within an AnimClip. Has Start and End times for the block,
+    which may or may not correspond to the node's natural size (i.e. having scaled versions of the block)
+    """
+    def __init__(self, animNode):
+        self.AnimNode = animNode
+        self.StartSeconds = 0.0
+        self.EndSeconds = 0.0
+        self.IsActive = True
+        self.Flipped = False
+
+    def getName(self):
+        return self.AnimNode.getName()
+    def setName(self, value):
+        self.AnimNode.setName(value)
+
+class DSAnimClip:
+    """
+    Represents a parsed animation clip (a vrNodePtr with type AnimClip or AnimWizClip), containing
+    a name and the AnimNodeUsages that are used by the clip. The AnimNodeUsages are store in a simple
+    list (so the track structure of animation blocks within clips is not respected).
+    """
+    def __init__(self, name):
+        self.Name = name
+        self._OriginalName = name
+        self.AnimNodeUsages = []
+        # Top level clips can be flipped if they're created with 'Duplicate and Flip'
+        self.Flipped = False
+
+    # Define these so that we can use the same sanitizer functions that
+    # we use for VRED scene nodes
+    def getName(self):
+        return self.Name
+    def setName(self, value):
+        self.Name = value
+    def getOriginalName(self):
+        return self._OriginalName
+    def getID(self):
+        return id(self)
+
+class DSAnimationLibrary:
+    """
+    This class serves as a temporary storage for animation data that we extract from the scene.
+    It is also capable of sanitizing the names of objects it creates, and finally it is responsible
+    for serializing itself to a file with an XML structure.
+    """
+    __DSID = 0
+
+    def __init__(self):
+        DSAnimationLibrary.__DSID = 0
+
+        self.TagTime = sys.float_info.max
+        self.BaseTime = 24.0
+        self.PlaybackSpeed = 24.0
+
+        self.DSAnimNodes = []
+        self.DSClips = []
+
+        # VRED guarantees all AnimNodes have unique names
+        self.BlockNameToDSBlock = {}
+
+        self.ClipOriginalNamesToNewNames = {}
+
+    @staticmethod
+    def getID():
+        """
+        Gets an unique identifier and internally increments it
+
+        Returns:
+            (int): Datasmith ID, uniquely identifies an animation-related object during export
+        """
+        DSAnimationLibrary.__DSID += 5
+        return DSAnimationLibrary.__DSID - 5
+
+    def sanitize(self):
+        """
+        Sanitizes the names of blocks and clips. Assumes node names are already sanitized
+        """
+        def sanitizeRecursive(clip, totalList):
+            try:
+                for usage in clip.AnimNodeUsages:
+                    sanitizeRecursive(usage, totalList)
+                totalList.append(clip)
+            except AttributeError:
+                pass
+
+        # AnimClips can never have repeated names, and even in nested clips you can't have
+        # clip references. So every single clip should have a unique name, as it will become
+        # an asset in UE4
+        clipsToSanitize = []
+        for clip in self.DSClips:
+            sanitizeRecursive(clip, clipsToSanitize)
+
+        dss.renameObjectsAndKeepOriginals(clipsToSanitize)
+
+        self.ClipOriginalNamesToNewNames = {}
+        for clip in clipsToSanitize:
+            self.ClipOriginalNamesToNewNames[clip.getOriginalName()] = clip.getName()
+
+        # The VRED AnimBlocks are sanitized before we even create the DSBlocks, so this
+        # might not be needed
+        blocksToSanitize = []
+        for node in self.DSAnimNodes:
+            blocksToSanitize += node.AnimBlocks
+        dss.renameObjectsAndKeepOriginals(blocksToSanitize)
+
+    def serialize(self, outFile):
+        """
+        Writes the animation library to 'outFile' as an XML file, describing the catalogged AnimNodes and AnimClips
+
+        Args:
+            outFile (file handle): File to write the library to
+        """
+        dsutils.writeXMLTag(outFile, 0, 'Root', [], [], False)
+
+        # Serialize blocks and curves
+        dsutils.writeXMLTag(outFile, 1, 'KeyTime', ['value'], [self.TagTime], True)
+        dsutils.writeXMLTag(outFile, 1, 'BaseTime', ['value'], [self.BaseTime], True)
+        dsutils.writeXMLTag(outFile, 1, 'PlaybackSpeed', ['value'], [self.PlaybackSpeed], True)
+        dsutils.writeXMLTag(outFile, 1, 'Blocks', [], [], False)
+        for node in self.DSAnimNodes:
+            dsutils.writeXMLTag(outFile, 2, 'AnimNode', ['name'], [node.Name], False)
+
+            for block in node.AnimBlocks:
+                dsutils.writeXMLTag(outFile, 3, 'AnimBlock', ['name'], [block.Name], False)
+
+                for curveName, curve in block.ChildCurveNamesToCurves.items():
+                    dsutils.writeXMLTag(outFile, 4, 'AnimCurve', ['name', 'DSID', 'startSeconds'], [curveName, curve.DSID, curve.StartSeconds], True)
+
+                dsutils.writeXMLTag(outFile, 3, '/AnimBlock', [], [], False)
+
+            dsutils.writeXMLTag(outFile, 2, '/AnimNode', [], [], False)
+
+        dsutils.writeXMLTag(outFile, 1, '/Blocks', [], [], False)
+
+        # Serialize clips and used blocks (and sub-clips)
+        dsutils.writeXMLTag(outFile, 1, 'Clips', [], [], False)
+        for clip in self.DSClips:
+            dsutils.writeXMLTag(outFile, 2, 'Clip', ['name', 'isFlipped'], [clip.Name, clip.Flipped], False)
+
+            for blockUsage in clip.AnimNodeUsages:
+                dsutils.writeXMLTag(outFile, 3, 'BlockUsage', ['blockName', 'startSeconds', 'endSeconds', 'isActive', 'isFlipped'], \
+                                    [blockUsage.getName(), blockUsage.StartSeconds, blockUsage.EndSeconds, blockUsage.IsActive, blockUsage.Flipped], True)
+
+            dsutils.writeXMLTag(outFile, 2, '/Clip', [], [], False)
+
+        dsutils.writeXMLTag(outFile, 1, '/Clips', [], [], False)
+
+        dsutils.writeXMLTag(outFile, 0, '/Root', [], [], False)
+
+    def catalogClip(self, animClip):
+        """
+        Recursively catalogs top-level animation clips and their child nodes.
+
+        Args:
+            animClip (vrNodePtr): Node that is of type AnimClip or AnimWizClip and contains other clips or AnimBlocks
+        """
+
+        def catalogClipInner(animClip, blockDestination):
+            """
+            Catalogs the AnimClip, recursively scanning so that all of its AnimBlocks children (and children of
+            children, etc) are added to the simple list in blockDestination.
+
+            Args:
+                animClip (vrNodePtr): Node that is of type AnimClip or AnimWizClip and contains other clips or AnimBlocks
+                blockDestination (list of AnimBlock vrNodePtr): Where to place the found AnimNode usages into
+            """
+            nChildren = animClip.getNChildren()
+            for childIndex in range(nChildren):
+                child = animClip.getChild(childIndex)
+                childName = child.getName()
+                childType = child.getType()
+
+                bb = child.getBoundingBox()
+
+                createdDSAnimNode = None
+
+                if childType == 'AnimClip' or childType == 'AnimWizClip':
+                    # Treat the sub-level clip like a top-level clip
+                    createdDSAnimNode = self.catalogClip(child)
+                elif childType == 'AnimBlock':
+                    # For animblocks we need to find the DSAnimBlock we're
+                    # referencing
+                    try:
+                        createdDSAnimNode = self.BlockNameToDSBlock[childName]
+                    except KeyError:
+                        logging.warning('Did not find AnimBlock \"' + str(childName) + '\"')
+                        continue
+                else:
+                    continue
+
+                childDSAnimBlockUsage = DSAnimNodeUsage(createdDSAnimNode)
+                childDSAnimBlockUsage.StartSeconds = bb[0]  # minX
+                childDSAnimBlockUsage.EndSeconds = bb[3]  # maxX
+                childDSAnimBlockUsage.IsActive = child.getActive()
+                childDSAnimBlockUsage.Flipped = dsutils.getIsNodeFlipped(child)
+
+                blockDestination.append(childDSAnimBlockUsage)
+
+        clipName = animClip.getName()
+
+        # Important to cause boundingBox evaluation. Without this start/end times would come back as zero
+        vrAnimWidgets.resetAnim(clipName)
+
+        DSClip = DSAnimClip(clipName)
+        DSClip.Flipped = dsutils.getIsNodeFlipped(animClip)
+        self.DSClips.append(DSClip)
+
+        catalogClipInner(animClip, DSClip.AnimNodeUsages)
+
+        return DSClip
+
+def getAnimRoot(node):
+    """
+    Returns the animationRoot of the AnimAttachment of 'node' if it has one.
+    Returns None otherwise.
+
+    Args:
+        node (vrNodePtr): Node to get the animRoot from
+
+    Returns:
+        (vrNodePtr): animationRoot of 'node' or None
+    """
+    animRoot = None
+    if node.hasAttachment("AnimAttachment"):
+        animAttachment = node.getAttachment("AnimAttachment")
+        animRootId = vrFieldAccess.vrFieldAccess(animAttachment).getFieldContainerID("animationRoot")
+        animRoot = vrNodePtr.vrNodePtr(animRootId)
+    return animRoot
+
+def splitAnimationFromGeometry(node):
+    """
+    If 'node' is an animated Geometry node, create a new parent node with its transformation and
+    animations, and add the node to it as a child, with its transformation reset
+
+    Args:
+        node (vrNodePtr): Node to split
+
+    Returns:
+        (vrNodePtr): parent node with animations and transformation (whose child is 'node'), or None
+    """
+    if node.getType() == 'Geometry' and getAnimRoot(node) is not None:
+        parentNode = vrScenegraph.createNode('Transform3D', node.getName() + '_Anim', node.getParent())
+        parentNode.addChild(node)
+
+        at = node.getAttachment('AnimAttachment')
+
+        parentNode.addAttachment(at)
+        node.subAttachment(at)
+        vrScenegraph.copyTransformation(node, parentNode)
+        dsutils.resetTransformNode(node)
+
+        logging.debug('Split animated geometry node \"' + str(node.getName()) + '\"')
+
+        return parentNode
+    return None
+
+def getFreeAnimCurves(animRoot):
+    """
+    Returns a list of nodes that are AnimCurves attached directly to animRoot, as opposed
+    to being attached to an AnimBlock that is attached to animRoot. Will only return curves
+    that have more than one key frame (i.e. valid animation curves)
+
+    Args:
+        animRoot (vrNodePtr): animationRoot node of an AnimAttachment
+
+    Returns:
+        (list of vrNodePtr): List of 'AnimCurve' nodes
+    """
+    result = []
+
+    nChildren = animRoot.getNChildren()
+    for childIndex in range(nChildren):
+        child = animRoot.getChild(childIndex)
+
+        if child.getType() == "AnimCurve":
+            bb = child.getBoundingBox()
+            # bb[0] and bb[3] are minX and maxX respectively, startSeconds and endSeconds of the curve
+            if not dsutils.isClose(bb[0], bb[3]):
+                result.append(child)
+
+    return result
+
+def getDSIDforChannelSuffix(channelName, suffix, curveDictionary):
+    """
+    Attempts to return the DSID of a curve in 'curveDictionary' that has the name 'channelName' + 'suffix'
+
+    Args:
+        channelName (str): Something like 'translation', 'scale' or 'visible'
+        suffix (str): Something like 'X', 'Y' or 'Z' or just ''
+        curveDictionary (dict of curve names to curves): Contains curves to search into
+
+    Returns:
+        (int): Found DSID, or zero if the name is not in dict
+    """
+    try:
+        return curveDictionary[channelName + suffix].DSID
+    except KeyError:
+        return 0
+
+def getDSIDVecForChannel(channelName, curveDictionary):
+    """
+    Attempts to return a Vec3f of DSIDs for the curves with 'channelName' within 'curveDictionary'
+
+    Args:
+        channelName (str): Something like 'translation', 'scale' or 'visible'
+        curveDictionary (dict of curve names to curves): Contains curves to search into
+
+    Returns:
+        (vrOSGTypes.Vec3f): Vector with the DSIDs of curves with suffixes 'X', 'Y' and 'Z' and 'channelName'
+    """
+    return vrOSGTypes.Vec3f(getDSIDforChannelSuffix(channelName, 'X', curveDictionary), \
+                            getDSIDforChannelSuffix(channelName, 'Y', curveDictionary), \
+                            getDSIDforChannelSuffix(channelName, 'Z', curveDictionary))
+
+def getKeyTime():
+    """
+    Scans all nodes in the scene and returns a time, in seconds, that is guaranteed
+    to be at least 1 second earlier than the start of any AnimCurve in the scene.
+
+    We use this to decide where to insert our DSID keys/value paris. We pack this in
+    the .clips file to be used when parsing those key/value pairs.
+
+    Returns:
+        (float): Time in seconds
+    """
+    minStartTime = sys.float_info.max
+
+    for node in vrScenegraph.getAllNodes():
+        animRoot = getAnimRoot(node)
+
+        if animRoot is None:
+            continue
+
+        numBlocks = animRoot.getNChildren()
+        for blockIndex in range(numBlocks):
+            block = animRoot.getChild(blockIndex)
+
+            numCurves = block.getNChildren()
+            for curveIndex in range(numCurves):
+                curve = block.getChild(curveIndex)
+
+                bb = curve.getBoundingBox()
+                startTimeSeconds = bb[0]
+
+                minStartTime = min(minStartTime, startTimeSeconds)
+
+    return minStartTime
+
+def catalogSceneAnimations(unitScale, baseTime, playbackSpeed):
+    """
+    Scan the scenegraph for all nodes that have AnimCurves, tags each curve with a marker,
+    then builds and returns an AnimationLibrary object that keeps track of Node, Clip and
+    Block animation names, as well as their corresponding DSID.
+
+    Args:
+        unitScale (float): The unit scale of the scene. 1.0 = mm, 10.0 = cm, 1000.0 = m
+                           This needs to be extracted from the UI, so its passed in from the plugin
+        baseTime (float): Native framerate of the animations in fps
+        playbackSpeed(float): Playback speed of the animations in fps. Animations that are natively
+                              1 second long are displayed with total duration 1s * (baseTime / playbackSpeed)
+    """
+    # Remove vertex animation blocks because those are not yet supported and will lead to the nodes being
+    # permanently locked on the END position of the animation
+    for node in vrScenegraph.getAllNodes():
+        animRoot = getAnimRoot(node)
+        if animRoot is None:
+            continue
+
+        numBlocks = animRoot.getNChildren()
+        for i in range(numBlocks-1, -1, -1):
+            block = animRoot.getChild(i)
+
+            if block.getType() == "TimeShapeCache":
+                logging.warning('Unsupported vertex animation block \"' + str(block.getName()) + '\". Removing...')
+                block.sub()
+
+    # Sanitize block names because if the user managed to create blocks with non-unique names
+    # all of this will fail
+    dss.renameObjectsAndKeepOriginals(vrAnimWidgets.getAnimBlockNodes(True))
+    library = DSAnimationLibrary()
+    library.BaseTime = baseTime
+    library.PlaybackSpeed = playbackSpeed
+
+    # Moves free curves to newly created nodename_BaseAnimation anim blocks
+    # Needs to happen before catalogging clips so that we keep track of these blocks too
+    for node in vrScenegraph.getAllNodes():
+        animRoot = getAnimRoot(node)
+
+        if animRoot is None:
+            continue
+
+        freeCurves = getFreeAnimCurves(animRoot)
+        if len(freeCurves) > 0:
+            # This block doesn't "work" within VRED, it will just serve as a container for us
+            newAnimBlock = vrScenegraph.createNode('AnimBlock', node.getName() + "_BaseAnimation", animRoot)
+            for curve in freeCurves:
+                newAnimBlock.addChild(curve)
+
+            # Register blocks by their names, to fill them with data later.
+            # Do this now so that we can use direct DSAnimBlock references within
+            # AnimNodeUsages, automatically renaming them when we sanitize the DSAnimBlock names. Additionally, doing this
+            # now allows us to export clips BEFORE we poison the curves with our tags
+            library.BlockNameToDSBlock[newAnimBlock.getName()] = DSAnimBlock(library.getID(), newAnimBlock.getName())
+
+    library.TagTime = getKeyTime() - 1.0
+
+    # We need to call resetAnim to calculate bounding boxes and get blockUsage durations for all clips.
+    # Creating an 'unanimated' block and using it to store our original transform allows us to reset to it later and
+    # export it with our original transform
+    unanimatedTransforms = {}
+    for node in vrScenegraph.getAllNodes():
+        animRoot = getAnimRoot(node)
+
+        if animRoot is None:
+            continue
+
+        rotOri = vrNodeUtils.getTransformNodeRotationOrientation(node)
+
+        storage = DSTransformStorage()
+        storage.Translation = node.getTranslation()
+        storage.Rotation = node.getRotation()
+        storage.Scale = node.getScale()
+        storage.Visible = node.getActive()
+        storage.RotationOrientation = [rotOri.x(), rotOri.y(), rotOri.z()]
+
+        unanimatedTransforms[node.getID()] = storage
+
+    # Register blocks by their names, to fill them with data later.
+    # Do this now so that we can use direct DSAnimBlock references within
+    # AnimNodeUsages, automatically renaming them when we sanitize the DSAnimBlock names. Additionally, doing this
+    # now allows us to export clips BEFORE we poison the curves with our tags
+    for block in vrAnimWidgets.getAnimBlockNodes(True) :
+        library.BlockNameToDSBlock[block.getName()] = DSAnimBlock(library.getID(), block.getName())
+
+    # Catalogue all clips and their blocks
+    # Needs to happen before we poison the curves with our tags, or else the AnimNodeUsages will have distorted start/end times
+    clips = vrAnimWidgets.getAnimClipNodes()
+    for clip in clips:
+        library.catalogClip(clip)
+
+    # Restore nodes to their unanimated states
+    for nodeID, storage in unanimatedTransforms.items():
+        node = vrNodePtr.vrNodePtr(nodeID)
+
+        node.setTranslation(storage.Translation[0], storage.Translation[1], storage.Translation[2])
+        node.setRotation(storage.Rotation[0], storage.Rotation[1], storage.Rotation[2])
+        node.setScale(storage.Scale[0], storage.Scale[1], storage.Scale[2])
+        node.setActive(storage.Visible)
+        vrNodeUtils.setTransformNodeRotationOrientation(node, storage.RotationOrientation[0], storage.RotationOrientation[1], storage.RotationOrientation[2])
+
+    # Catalogs all nodes, their blocks and their curves
+    for node in vrScenegraph.getAllNodes():
+        animRoot = getAnimRoot(node)
+
+        if animRoot is None:
+            continue
+
+        # Use this bool to prevent creating DSAnimNodes for nodes without animBlocks
+        createdANodeAlready = False
+
+        # Tag curves in every block, block by block
+        numBlocks = animRoot.getNChildren()
+        for blockIndex in range(numBlocks):
+            block = animRoot.getChild(blockIndex)
+            if block.getType() != "AnimBlock":
+                continue
+
+            if not createdANodeAlready:
+                animNode = DSAnimNode(node.getName())
+                library.DSAnimNodes.append(animNode)
+                createdANodeAlready = True
+
+            dsBlock = None
+            try:
+                dsBlock = library.BlockNameToDSBlock[block.getName()]
+            except KeyError:
+                logging.warning('Animation block with name \"' + str(block.getName()) + '\" does not exist!')
+                pass
+            if dsBlock is None:
+                continue
+
+            animNode.AnimBlocks.append(dsBlock)
+            numCurves = block.getNChildren()
+
+            for curveIndex in range(numCurves-1, -1, -1):
+                curve = block.getChild(curveIndex)
+                if curve.getType() != "AnimCurve":
+                    continue
+
+                vrScenegraph.moveNode(curve, block, animRoot)
+
+                bb = curve.getBoundingBox()
+                startTimeSeconds = bb[0]
+
+                # We want all rotation curves to have the same dsid, and
+                # all scale curves to have another (but consistent) id
+                dsid = dsBlock.DSID
+                channelName = curve.fields().getString('channelName')
+                if channelName == "translation":
+                    pass
+                elif channelName == "rotation":
+                    dsid += 1
+                elif channelName == "scale":
+                    dsid += 2
+                elif channelName == "rotationOrientation":
+                    dsid += 3
+                elif channelName == "visible":
+                    dsid += 4
+                else:
+                    dsid = -1
+
+                dsCurve = DSAnimCurve(dsid, curve, startTimeSeconds)
+                dsBlock.ChildCurveNamesToCurves[curve.fields().getString('channelName') + curve.fields().getString('suffix')] = dsCurve
+
+            # Get tag vectors
+            transVec = getDSIDVecForChannel('translation', dsBlock.ChildCurveNamesToCurves)
+            rotVec   = getDSIDVecForChannel('rotation', dsBlock.ChildCurveNamesToCurves)
+            scaleVec = getDSIDVecForChannel('scale', dsBlock.ChildCurveNamesToCurves)
+            visible  = getDSIDforChannelSuffix('visible', '', dsBlock.ChildCurveNamesToCurves)
+
+            frameTime = 1.0 / baseTime
+
+            # We used to have a different one for each component, but now we just take x(),
+            # as they will be the same
+            transTime = library.TagTime - transVec.x() * frameTime
+            rotTime = library.TagTime - rotVec.x() * frameTime
+            scaleTime = library.TagTime - scaleVec.x() * frameTime
+            visibleTime = library.TagTime - visible * frameTime
+
+            someValue = vrOSGTypes.Vec3f(0.0, 0.0, 0.0)
+
+            # Apply tags
+            if not vrAnimWidgets.addTranslationControlPoint(node, transTime, someValue, False):
+                logging.warning('Failed to add translation tags for block ' + str(block.getName()))
+            if not vrAnimWidgets.addRotationControlPoint(node, rotTime, someValue):
+                logging.warning('Failed to add rotation tags for block ' + str(block.getName()))
+            if not vrAnimWidgets.addScaleControlPoint(node, scaleTime, someValue):
+                logging.warning('Failed to add scale tags for block ' + str(block.getName()))
+            if not vrAnimWidgets.addVisibleControlPoint(node, visibleTime, 1.0):
+                logging.warning('Failed to add visible tags for block ' + str(block.getName()))
+
+            # Move them back to the block so that they will be exported modified
+            for dsCurve in dsBlock.ChildCurveNamesToCurves.values():
+                vrScenegraph.moveNode(dsCurve.AnimCurve, animRoot, block)
+
+    return library
+
diff --git a/Engine/Plugins/Enterprise/DatasmithFBXImporter/Resources/VREDPlugin/DatasmithExporter/exporter.py b/Engine/Plugins/Enterprise/DatasmithFBXImporter/Resources/VREDPlugin/DatasmithExporter/exporter.py
new file mode 100644
--- /dev/null
+++ b/Engine/Plugins/Enterprise/DatasmithFBXImporter/Resources/VREDPlugin/DatasmithExporter/exporter.py
@@ -0,0 +1,1456 @@
+from PySide2 import QtCore, QtGui, QtWidgets
+from PySide2.QtWidgets import QMessageBox
+
+import sys
+import os
+import logging
+import xml.etree.ElementTree as ET
+from collections import defaultdict
+from itertools import izip_longest, chain
+
+import vrFileIO
+import vrFileDialog
+import vrVariantSets
+import vrMaterialPtr
+import vrScenegraph
+import vrNodeUtils
+import vrFieldAccess
+import vrController
+import vrOptimize
+import vrNodePtr
+import vrOSGWidget
+import vrCamera
+
+# Adds the ScriptPlugins folder to path so that we can find the other files
+# in our module
+scriptPluginsFolder = os.path.dirname(os.path.realpath(os.curdir))
+if scriptPluginsFolder not in sys.path:
+    sys.path.append(scriptPluginsFolder)
+
+import DatasmithExporter.utils as dsutils
+import DatasmithExporter.anim_exporter as dsanim
+import DatasmithExporter.sanitizer as dss
+import DatasmithExporter.mat as dsmat
+
+# Current plugin version
+exporterVersion = "3"
+
+class Light:
+    def __init__(self):
+        self.name = ""
+        self.type = ""
+        self.enabled = True
+        self.useTemperature = False
+        self.temperature = 0.0
+        self.visualizationVisible = True
+        self.intensity = 0.0
+        self.glossyColor = []
+        self.diffuseColor = []
+        self.coneAngle = 0.0
+        self.penumbraAngle = 0.0
+        self.areaLightUseConeAngle = False
+        self.attenuationType = 0
+        self.unit = 0
+        self.useIESProfile = False
+        self.IESPath = ""
+
+class UserSettings:
+    def __init__(self):
+        self.animBaseTime = 24.0
+        self.animPlayRate = 24.0
+
+    def tryFetchFromFile(self):
+        try:
+            appdata = os.getenv('APPDATA')
+            path = os.path.join(appdata, r"VREDPro\\preferences.xml")
+
+            if os.path.exists(path):
+                settingsFile = ET.parse(path)
+                settingsRoot = settingsFile.getroot()
+
+                timeNodes = settingsRoot.findall('./key[@name="Animation time"]')
+                self.animBaseTime = float(timeNodes[0].text)
+
+                playbackRateNodes = settingsRoot.findall('./key[@name="Animation playback speed"]')
+                playbackMode = int(playbackRateNodes[0].text)
+                if playbackMode == 0:  # Every Frame ~ 60fps
+                    self.animPlayRate = 60.0
+                elif playbackMode == 1:  # Film
+                    self.animPlayRate = 24.0
+                elif playbackMode == 2:  # PAL
+                    self.animPlayRate = 25.0
+                elif playbackMode == 3:  # NTSC
+                    self.animPlayRate = 30.0
+                else:  # Custom
+                    customRateNodes = settingsRoot.findall('./key[@name="Animation playback custom speed"]')
+                    self.animPlayRate = float(customRateNodes[0].text)
+
+                logging.info('Found preferences file "'  + str(path) + '"')
+                logging.info('Read animation base time of ' + str(self.animBaseTime) + ' fps')
+                logging.info('Read animation playback speed of ' + str(self.animPlayRate) + ' fps')
+
+        except BaseException as e:
+            self.animBaseTime = 24.0
+            self.animPlayRate = 24.0
+            logging.warning('Problem accessing preferences file: "' + repr(e) + '"')
+            logging.info('Set animation base time to default of 24.0 fps')
+            logging.info('Set animation playback speed to default of 24.0 fps')
+
+
+class DatasmithFBXExporter():
+    def __init__(self, mainWindow):
+        self.mainWindow = mainWindow
+
+    def getMaterialsToPreserve(self):
+        """
+        If an object has a Switch Material currently set to MaterialA, when exported to the
+        FBX it will be written as simply just having MaterialA. This means that the unselected variants
+        of a Switch Material may not be used by any node, which causes them to not be packed into the
+        FBX at all. This function identifies these unused material variants and returns a list with
+        the material pointers
+
+        Returns:
+            (list of vrMaterialPtr): List of vrMaterialPtr of materials to preserve
+        """
+        # Go over all switch materials and keep track of the IDs of
+        # submaterials that are not currently chosen
+        idsOfMatsToPreserve = []
+        for mat in vrMaterialPtr.getAllMaterials():
+            if mat.getNodes() == 0:
+                continue
+
+            submats = mat.getSubMaterials()
+            if len(submats) > 1:
+                chosenId = submats[mat.getChoice()].getID()
+                for submat in submats:
+                    if submat.getID() != chosenId:
+                        idsOfMatsToPreserve.append(submat.getID())
+
+        # Go over the list again discarding all materials that are used by every switch.
+        # This because a material that is unused by one switch might be used by another, and so
+        # doesn't actually need to be preserved
+        for mat in vrMaterialPtr.getAllMaterials():
+            if mat.getNodes() == 0:
+                continue
+
+            submats = mat.getSubMaterials()
+
+            # If its a switch material, try removing its selected variant from the list
+            if len(submats) > 1:
+                chosenId = submats[mat.getChoice()].getID()
+                try:
+                    idsOfMatsToPreserve.remove(chosenId)
+                except ValueError:
+                    pass
+
+            # If its a regular material used by an active node, try removing it from the list
+            # since it will be packed in the FBX already
+            else:
+                try:
+                    for node in mat.getNodes():
+                        if node.getActive():
+                            idsOfMatsToPreserve.remove(mat.getID())
+                except ValueError:
+                    pass
+
+        return [vrMaterialPtr.toMaterial(matId) for matId in idsOfMatsToPreserve]
+
+    def assignMatsToTempGeometry(self, mats):
+        """
+        For each vrMaterialPtr in mats, creates a temp plane and assigns the material to it,
+        returning a list of all vrNodePtr that it created.
+
+        Args:
+            mats (list of vrMaterialPtr): Materials to assign to temp geometry
+
+        Returns:
+            List of vrNodePtr containing the temp geometry created
+        """
+        newNodes = []
+        for mat in mats:
+            newNode = vrNodeUtils.createPlane(1, 1, 1, 1, 0, 0, 0)
+            newNode.setName("__temp_" + str(id(mat)))  # Use python object ID for uniqueness
+            newNode.setActive(True)
+            newNode.setMaterial(mat)
+            newNodes.append(newNode)
+
+        # Disable in scene graph and hide in perspective view
+        # will still be briefly visible in the graph though
+        vrScenegraph.hideNodes(newNodes)
+
+        return newNodes
+
+    def injectSwitchMaterialUsers(self, xmlRoot):
+        """
+        If an object has a Switch Material currently set to MaterialA, when exported to the
+        FBX it will be written as simply just having MaterialA. We want to preserve that
+        these objects had a switch material, so we'll find the Switch Material entry in the
+        variants xml tree and inject the name of the nodes that use them with <UsedBy name="nodeName"/>
+        tags
+
+        Assumes that switch materials have already been renamed in the variants file
+
+        Args:
+            xmlRoot (xml.etree.ElementTree.Element): Root of the variants XML tree
+        """
+        for matNode in xmlRoot.iterfind('.//MaterialVariant'):
+            # 'base' is the actual material name
+            materialName = matNode.get('base')
+
+            logging.debug('Injecting users for switchmat variant referencing material ' + str(materialName))
+
+            matPtrs = vrMaterialPtr.findMaterials(materialName)
+            if len(matPtrs) != 1:
+                logging.warning('Found ' + str(len(matPtrs)) + ' switch materials with name ' + str(materialName) + ' when injecting users!')
+
+            if len(matPtrs) == 0:
+                continue
+
+            mat = matPtrs[0]
+            for node in mat.getNodes():
+                logging.debug('\t\tinjecting user ' + str(node.getName()))
+                ET.SubElement(matNode, 'UsedBy', {'name':node.getName()})
+
+    def extractTransformVariants(self):
+        """
+        Uses the vrFileAccess module to extract TransformVariants from nodes that have them.
+        It will return a dictionary of node names to other dictionaries, these in turn containing
+        the actual transform variants keyed by their names. Example:
+
+        {'BoxNode': {'Variant1': [1, 0, 0, ..., 1], 'Variant2': [1, 2, ..., 1]},
+         'SphereNode': {'Variant1': [1, 0, 0, ..., 1]}}
+
+        Returns:
+            (dict): Keys are node names. Each value is a dict, where variant names are keys, and
+                    values are float arrays describing the actual transform variants
+        """
+        nodes = vrScenegraph.getAllNodes()
+
+        unitScale = self.getUnitScale()
+
+        nodeNamesToTransVars = {}
+        for node in nodes:
+            nodeName = node.getName()
+
+            # Some objects have the TransformVariantsAttachment attached directly to them, some have it
+            # attached to their transform field container
+            if node.fields().hasField('transform'):
+                transformContainer = vrFieldAccess.vrFieldAccess(node.fields().getFieldContainer('transform'))
+            else:
+                transformContainer = node.fields()
+
+            # Carefully drill down to the transform variant data
+            if transformContainer.isValid() and transformContainer.hasAttachment('TransformVariantsAttachment'):
+                variantsAttachment = transformContainer.getAttachment('TransformVariantsAttachment')
+                variantsAttachment = vrFieldAccess.vrFieldAccess(variantsAttachment)
+
+                if variantsAttachment.isValid() and variantsAttachment.hasField('transformVariants'):
+                    variantsList = variantsAttachment.getMFieldContainer('transformVariants')
+
+                    # Iterate over all TransformVariants that this node has
+                    for variant in variantsList:
+                        var = vrFieldAccess.vrFieldAccess(variant)
+
+                        transformationContainer = vrFieldAccess.vrFieldAccess(var.getFieldContainer('transformation'))
+
+                        # Pack transform data in output dictionary
+                        transName = var.getString('name')
+                        transform = self.processTransform(transformationContainer.getMatrix('matrix'), unitScale)
+
+                        if nodeName not in nodeNamesToTransVars:
+                            nodeNamesToTransVars[nodeName] = defaultdict(list)
+
+                        nodeNamesToTransVars[nodeName][transName].append(transform)
+
+        return nodeNamesToTransVars
+
+    def processTransform(self, trans, scale):
+        """
+        Flip and transpose the transform so it can be used directly in UE4. Also, scales
+        the translation by the unit scale, which is something that's easier to do here than
+        in the UE4 importer.
+
+        Args:
+            trans (list of 16 floats): Describes the transform to process
+            scale (float): Scale unit factor, where 1.0 means mm, 10.0 means cm, 1000.0 means m, etc
+
+        Returns:
+            (list of 16 floats): Processed transform
+        """
+        assert(len(trans) == 16)
+
+        # Transpose transform
+        transposed = [0]*16
+        index = 0
+        for row in range(0, 4):
+            for col in range(0, 4):
+                transposed[index] = trans[col*4 + row]
+                index += 1
+
+        # Flip around the Y plane
+        for pos in range(0, 4):
+            transposed[pos*4 + 1] *= -1
+        for pos in range(0, 4):
+            transposed[4 + pos] *= -1
+
+        # # Scale translation by unit
+        multiplier = scale / 10
+        for pos in range(0, 3):
+            transposed[12 + pos] *= multiplier
+
+        return transposed
+
+    def injectTransformVariants(self, xmlRoot, nodeNamesToTransVars):
+        """
+        Injects the actual transform float arrays in the variants file describing transform
+        variants.
+
+        Args:
+            xmlRoot (xml.etree.ElementTree.Element): Root of the variants XML tree
+            nodeNamesToTransVars (dict): Keys are node names. Each value is a dict, where
+                                         variant names are keys, and values are float arrays
+                                         describing the actual transform variants
+        """
+        # We'll use try/except KeyError below, so its better to make sure its not a defaultdict
+        nodeNameToTransVarsDict = dict(nodeNamesToTransVars)
+
+        logging.debug('Injecting transform variants ' + str(nodeNameToTransVarsDict))
+
+        for transVarXmlNode in xmlRoot.iterfind('./TransformVariant'):
+            nodeName = transVarXmlNode.get('base')
+            for state in transVarXmlNode.findall('State'):
+                stateName = state.get('name')
+
+                logging.debug('\tAnalyzing transformVariants for node ' + str(nodeName) + ', state ' + str(stateName))
+
+                if nodeName in nodeNamesToTransVars:
+                    if stateName in nodeNamesToTransVars[nodeName]:
+                        transVars = nodeNamesToTransVars[nodeName][stateName]
+
+                        state.set('value', str(transVars[0]))
+                        logging.debug('\tinjecting ' + str(transVars[0]) + ' to state ' + str(stateName) + ' of transVar of node ' + str(nodeName))
+                        del transVars[0]
+
+                        # If we have multiple transform variant nodes, VRED won't have inserted an extra line in
+                        # the variants file, let's fix that since its easy
+                        while len(transVars) != 0:
+                            ET.SubElement(transVarXmlNode, 'State', {'name':stateName, 'value':str(transVars[0])})
+                            logging.debug('\tinjecting additional transVar ' + str(transVars[0]) + ' to state ' + str(stateName) + ' of transVar of node ' + str(nodeName))
+                            del transVars[0]
+                    else:
+                        if stateName not in ['!Previous', '!Next', '!Previous (Loop)', '!Next (Loop)']:
+                            logging.warning('Did not find transform for transVar state ' + str(stateName) + ' of node ' + str(nodeName))
+                else:
+                    logging.warning('Did not find transform of node ' + str(nodeName))
+
+    def sanitizeTransformVariantStateNames(self, xmlRoot, transVariants):
+        """
+        Sanitizes the names of transform variant states, updating the default to match
+
+        Args:
+            xmlRoot (xml.etree.ElementTree.Element): Root of the variants XML tree
+        """
+        commandNames = set(['!Next', '!Next (Loop)', '!Previous', '!Previous (Loop)'])
+
+        renamedNodeVariantStates = defaultdict(dict)
+        for transVarXmlNode in xmlRoot.iterfind('./TransformVariant'):
+            nodeName = transVarXmlNode.get('base')
+            logging.debug('Analyzing transform variant state names for node ' + str(nodeName))
+
+            adapters = [dss.XMLNodeAdapter(xmlNode) for xmlNode in transVarXmlNode.findall('State') if xmlNode.get('name') not in commandNames]
+            adapterDict = {ad.getID(): ad for ad in adapters}
+
+            # Rename states
+            # Note how since we're using a dict here two states with the same name will shadow eachother,
+            # and our varset refs will end up all pointing at the same one.
+            # That is fine though, since they wouldn't have worked fine in VRED anyway
+            renamedAdapters = dss.renameObjectsAndKeepOriginals(adapters)
+            stateOldToNewNames = {orig: adapterDict[adId].getName() for adId, orig in renamedAdapters.iteritems()}
+            logging.debug('\trenamed states: ' + str(stateOldToNewNames))
+
+            # Keep track of renamed names so that we can update varset references
+            transVarName = transVarXmlNode.get('name')
+            renamedNodeVariantStates[transVarName] = stateOldToNewNames
+
+            # Rename default state
+            for defNode in transVarXmlNode.findall('Default'):
+                stateName = defNode.get('ref')
+
+                if stateName in stateOldToNewNames:
+                    newStateName = stateOldToNewNames[stateName]
+                    defNode.set('ref', newStateName)
+                    logging.debug('\trenaming default state from  ' + str(stateName) + ' to ' + str(newStateName))
+
+        # Rename references to node variants from variant sets
+        for xmlNode in xmlRoot.iterfind(".//TransformRef"):
+            maybeRenamedRef = xmlNode.get('ref')
+            oldStateName = xmlNode.get('state')
+
+            logging.debug('Analyzing TransformRef ' + str(maybeRenamedRef))
+
+            if maybeRenamedRef in renamedNodeVariantStates:
+                renamedStates = renamedNodeVariantStates[maybeRenamedRef]
+
+                if oldStateName in renamedStates:
+                    newStateName = renamedStates[oldStateName]
+                    logging.debug('\tstate ' + str(oldStateName) + ' seems to have been renamed to ' + str(newStateName))
+                    xmlNode.set('state', newStateName)
+
+    def extractLights(self):
+        """
+        Extract all lights from the current scene, packs each into a Light object
+        and returns a list of them
+
+        Returns:
+            (list of Light objects): List of lights in the scene
+        """
+        lights = []
+        nodes = vrScenegraph.getAllNodes()
+        for node in nodes:
+            if node.isLight():
+                light = Light()
+                light.name = node.getName()
+                light.worldTransform = node.getWorldTransform()
+
+                if node.isSpotLight():
+                    light.type = "spotlight"
+                elif node.isSphericalLight():
+                    light.type = "spherical"
+                elif node.isRectangularLight():
+                    light.type = "rectangular"
+                elif node.isPointLight():
+                    light.type = "point"
+                elif node.isDiskLight():
+                    light.type = "disk"
+                elif node.isDirectionalLight():
+                    light.type = "directional"
+                elif node.isAreaLight():
+                    light.type = "area"
+
+                light.enabled = node.isLightOn()
+                light.useTemperature = node.getUseLightTemperature()
+                light.temperature = node.getLightTemperature()
+                light.visualizationVisible = node.getLightVisualizationVisible()
+                light.intensity = node.getLightIntensity()
+                light.glossyColor = node.getGlossyLightColor()
+                light.diffuseColor = node.getDiffuseLightColor()
+                light.coneAngle = node.getLightConeAngle()
+                light.penumbraAngle = node.getLightPenumbraAngle()
+                light.areaLightUseConeAngle = node.getAreaLightUseConeAngle()
+                light.useIESProfile = node.getUseIESProfile()
+
+                # Extract attenuationType using the fieldaccess API
+                try:
+                    if node.fields().hasField('attenuationType'):
+                        light.attenuationType = node.fields().getUInt32('attenuationType')
+                    else:
+                        light.attenuationType = 0
+                except AttributeError:
+                    light.attenuationType = 0
+
+                # Extract unit using the fieldaccess API
+                try:
+                    if node.fields().hasField('unit'):
+                        light.unit = node.fields().getUInt32('unit')
+                    else:
+                        light.unit = 0
+                except AttributeError:
+                    light.unit = 0
+
+                # Extract IES filename using the fieldaccess API
+                try:
+                    fc = vrFieldAccess.vrFieldAccess(node.fields().getFieldContainer('iesProfile'))
+                    if fc.isValid():
+                        light.IESPath = fc.getString('filename')
+                    else:
+                        light.IESPath = ""
+                except AttributeError:
+                    light.IESPath = ""
+
+                lights.append(light)
+        return lights
+
+    def writeLightsFile(self, filename, lights):
+        """
+        Creates a file with 'filename' and repeatedly calls writeXMLTag to write all the
+        lights in the 'lights' array to the XML file
+
+        Args:
+            filename (str): Full path of the file to write
+            lights (list of Light objects): List of lights to write to the file
+        """
+
+        with open(filename, 'w') as lightsfile:
+            dsutils.writeXMLTag(lightsfile, 0, 'Lights', [], [])
+
+            for light in lights:
+                dsutils.writeXMLTag(lightsfile, 1, 'Light', ['name'], [light.name])
+
+                dsutils.writeXMLTag(lightsfile, 2, 'WorldTransform', ['value'], [light.worldTransform], close=True)
+                dsutils.writeXMLTag(lightsfile, 2, 'Type', ['value'], [light.type], close=True)
+                dsutils.writeXMLTag(lightsfile, 2, 'Enabled', ['value'], [light.enabled], close=True)
+                dsutils.writeXMLTag(lightsfile, 2, 'UseTemperature', ['value'], [light.useTemperature], close=True)
+                dsutils.writeXMLTag(lightsfile, 2, 'Temperature', ['value'], [light.temperature], close=True)
+                dsutils.writeXMLTag(lightsfile, 2, 'Intensity', ['value'], [light.intensity], close=True)
+                dsutils.writeXMLTag(lightsfile, 2, 'GlossyColor', ['value'], [light.glossyColor], close=True)
+                dsutils.writeXMLTag(lightsfile, 2, 'DiffuseColor', ['value'], [light.diffuseColor], close=True)
+                dsutils.writeXMLTag(lightsfile, 2, 'ConeAngle', ['value'], [light.coneAngle], close=True)
+                dsutils.writeXMLTag(lightsfile, 2, 'PenumbraAngle', ['value'], [light.penumbraAngle], close=True)
+                dsutils.writeXMLTag(lightsfile, 2, 'AreaLightUseConeAngle', ['value'], [light.areaLightUseConeAngle],
+                                 close=True)
+                dsutils.writeXMLTag(lightsfile, 2, 'VisualizationVisible', ['value'], [light.visualizationVisible],
+                                 close=True)
+                dsutils.writeXMLTag(lightsfile, 2, 'AttenuationType', ['value'], [light.attenuationType], close=True)
+                dsutils.writeXMLTag(lightsfile, 2, 'Unit', ['value'], [light.unit], close=True)
+                dsutils.writeXMLTag(lightsfile, 2, 'UseIESProfile', ['value'], [light.useIESProfile], close=True)
+                dsutils.writeXMLTag(lightsfile, 2, 'IESPath', ['value'], [light.IESPath], close=True)
+
+                dsutils.writeXMLTag(lightsfile, 1, '/Light', [], [])
+
+            dsutils.writeXMLTag(lightsfile, 0, '/Lights', [], [])
+
+    def writeTextures(self, folder, mats):
+        """
+        Writes all textures in 'mats' as png files within folder, also updates tex objects
+        with the exported filepath
+
+        Args:
+            folder (str): Full path folder to contain all pngs (e.g. C:/TextureFolder)
+            mats (list of Mat objects): List of mats to write to the file
+        """
+        exportedTextures = {}
+
+        for mat in mats:
+            for texName in mat.textures:
+                tex = mat.textures[texName]
+
+                srcName = os.path.split(tex.srcPath)[1]
+                srcName = os.path.splitext(srcName)[0]
+                if len(srcName) == 0:
+                    srcName = str(mat.vredMatId) + str(texName)
+                filename = os.path.join(folder, srcName)
+                filename += '_' + str(tex.texType).replace('Component', '')
+
+                # If we already have a texture from the same file used in the same way,
+                # give up. The way it is used is important as VRED does some processing for
+                # displacement/normal maps
+                # Disregard how this will be like 'file.pngglossyColor', this is just a key
+                usageKey = str(tex.srcPath) + str(tex.texType)
+                if usageKey in exportedTextures:
+                    tex.exportedPath = exportedTextures[usageKey]
+                    continue
+
+                arr = bytearray(tex.pixelData)
+
+                # If it's invalid we'll just dump the data directly. Some environment/hdr formats
+                # are RGBA64 for example, which this version of Qt doesn't support
+                imgFormat = dsutils.getQImageFormat(tex.bytesPerPixel, tex.channels)
+                if imgFormat == QtGui.QImage.Format_Invalid or tex.dataType != 5121:
+                    filename += '.tex'
+
+                    logging.warning('Texture "' + str(srcName) + '" has an unsupported format (width: ' + str(tex.width) + ', height: ' + str(tex.height) + ', bytesPerPixel: ' + str(tex.bytesPerPixel) + ', pixelFormat: ' + str(tex.pixelFormat) + ', dataType: "' + str(tex.dataType) + '"), so it will be dumped in raw binary format to "' + str(filename) + '". It will have no extra padding per row.')
+
+                    with open(filename, 'wb') as f:
+                        f.write(arr)
+                else:
+                    filename += '.png'
+                    img = QtGui.QImage(arr, tex.width, tex.height, tex.stride, imgFormat)
+                    img.save(filename, "PNG", 100)
+
+                if os.path.exists(filename):
+                    tex.exportedPath = filename
+                    exportedTextures[usageKey] = filename
+
+    def writeMatFile(self, filename, mats):
+        """
+        Creates a file with 'filename' and repeatedly calls writeXMLTag to write all the
+        mats in the 'mats' array to the XML file
+
+        Args:
+            filename (str): Full path of the file to write
+            mats (list of Mat objects): List of mats to write to the file
+        """
+
+        with open(filename, 'w') as f:
+            dsutils.writeXMLTag(f, 0, 'Mats', [], [])
+
+            for mat in mats:
+                dsutils.writeXMLTag(f, 1, 'Mat', ['name', 'matType', 'vredMatId'], [mat.name, mat.matType, mat.vredMatId])
+
+                dsutils.writeXMLTag(f, 2, 'Users', [], [])
+                for user in mat.users:
+                    dsutils.writeXMLTag(f, 3, 'User', ['name'], [user], close=True)
+                dsutils.writeXMLTag(f, 2, '/Users', [], [])
+
+                dsutils.writeXMLTag(f, 2, 'Submaterials', [], [])
+                for submat in mat.submaterials:
+                    dsutils.writeXMLTag(f, 3, 'Submaterial', ['name'], [submat], close=True)
+                dsutils.writeXMLTag(f, 2, '/Submaterials', [], [])
+
+                dsutils.writeXMLTag(f, 2, 'Properties', [], [])
+                for prop in mat.properties:
+                    dsutils.writeXMLTag(f, 3, prop, ['value'], [mat.properties[prop]], close=True)
+                dsutils.writeXMLTag(f, 2, '/Properties', [], [])
+
+                dsutils.writeXMLTag(f, 2, 'Textures', [], [])
+                for texName in mat.textures:
+                    tex = mat.textures[texName]
+                    dsutils.writeXMLTag(f, 3, 'Tex', ['exportedPath', 'srcPath', 'texType', 'width', 'height', 'bytesPerPixel', 'channels', 'dataType', 'stride'], [tex.exportedPath, tex.srcPath, tex.texType, tex.width, tex.height, tex.bytesPerPixel, tex.channels, tex.dataType, tex.stride])
+
+                    dsutils.writeXMLTag(f, 4, 'Properties', [], [])
+                    for texProp in tex.properties:
+                        dsutils.writeXMLTag(f, 5, texProp, ['value'], [tex.properties[texProp]], close=True)
+                    dsutils.writeXMLTag(f, 4, '/Properties', [], [])
+
+                    dsutils.writeXMLTag(f, 3, '/Tex', [], [])
+
+                dsutils.writeXMLTag(f, 2, '/Textures', [], [])
+
+                dsutils.writeXMLTag(f, 1, '/Mat', [], [])
+
+            dsutils.writeXMLTag(f, 0, '/Mats', [], [])
+
+    def removePreviewTags(self, xmlRoot):
+        """
+        VRED encodes variant set preview thumbnails in base64 in the variants file itself.
+        Those are massive, unused and just hang the re evaluation, so this function removes them.
+
+        For performance, make sure this is called as soon as you read back the variants file.
+
+        Args:
+            xmlRoot (xml.etree.ElementTree.Element): Root of the variants XML tree
+        """
+        for node in xmlRoot.iterfind(".//Preview/.."):
+            for prev in node.iterfind('Preview'):
+                node.remove(prev)
+
+    def sanitizeVariantNamesAndVariantSetRefs(self, xmlRoot):
+        """
+        Sanitizes names of variants and variant sets directly on the xml tree.
+        We will only sanitize variant name and variant set ref fields, since 'state' and 'base'
+        fields correspond to actual nodes and material nodes, sanitized and renamed elsewhere
+
+        Args:
+            xmlRoot (xml.etree.ElementTree.Element): Root of the variants XML tree
+        """
+        allVariantNodes = []
+
+        # This isn't so bad since we only ever go one level down anyway
+        allVariantNodes += xmlRoot.findall("./NodeVariant")
+        allVariantNodes += xmlRoot.findall("./MaterialVariant")
+        allVariantNodes += xmlRoot.findall("./LightVariant")
+        allVariantNodes += xmlRoot.findall("./TransformVariant")
+        allVariantNodes += xmlRoot.findall("./VariantSet")
+        allVariantNodes += xmlRoot.findall("./VariantGroup/VariantSet")
+
+        # Use adapters so that we can reuse our sanitizing function
+        variantNodeAdapters = [dss.XMLNodeAdapter(xmlNode) for xmlNode in allVariantNodes]
+
+        # We'll sanitize all of them together since they will become rows in a datatable in Unreal, meaning
+        # they must have unique names even between different types of variants
+        renamedVariants = dss.renameObjectsAndKeepOriginals(variantNodeAdapters)
+
+        # Group all adapters by originalName
+        origNamesToAdapters = defaultdict(list)
+        for adapter in variantNodeAdapters:
+            adapterID = adapter.getID()
+            if adapterID in renamedVariants:
+                origName = renamedVariants[adapterID]
+                origNamesToAdapters[origName].append(adapter)
+        origNamesToAdapters = dict(origNamesToAdapters)
+
+        logging.debug('Found the renamed variant references ' + str(origNamesToAdapters))
+
+        for variantRefNode in xmlRoot.iterfind(".//VariantSet/*[@ref]"):
+            oldRef = variantRefNode.get('ref')
+            refTag = variantRefNode.tag.replace('Ref', '').replace('Variant', '')  # Node, Material, Transform, etc
+
+            if oldRef in origNamesToAdapters:
+                adapters = origNamesToAdapters[oldRef]
+
+                for adapter in adapters:
+                    variantType = adapter.getType()
+                    variantTypePruned = variantType.replace('Variant', '')  # Also Node, Material, Transform, etc
+
+                    logging.debug('\tcomparing ref type ' + str(refTag) + ' with var type ' + str(variantTypePruned))
+
+                    if refTag == variantTypePruned:
+                        newRef = adapter.getName()
+                        variantRefNode.set('ref', newRef)
+                        logging.debug('\tmatch! renaming ' + str(oldRef) + ' to ' + str(newRef))
+
+    def unshareNodes(self):
+        """
+        Unshares all cloned and transformable roots in the scene. This is not an undoable operation.
+        """
+        # Every time we unshare a node, *something* happens to the scenegraph that invalidates
+        # all the references and vrNodePtr that have to other nodes. For some reason, we need
+        # to call vrScenegraph.getAllNodes() and get new references, or else we end up with invalid/broken nodes
+        # Particularly, unshared nodes seem to not each get a copy of the cloned 'Name' attachment, so when
+        # we call .getName() we get a full crash
+
+        def stillHaveClones():
+            """
+            Gets the first shared node in the scene that it finds, then climbs the hierarchy to find the
+            root of the cloned subtree and unshares it.
+
+            Returns:
+                (bool): False if there are no more shared nodes in the scenegraph
+            """
+            try:
+                aClone = next(node for node in vrScenegraph.getAllNodes() if dsutils.isSharedNode(node))
+            except StopIteration:
+                return False
+
+            parent = aClone.getParent()
+            while parent != None and parent.isValid() and parent != aClone and dsutils.isSharedNode(parent):
+                aClone = parent
+                parent = aClone.getParent()
+
+            vrScenegraph.unshareNode(aClone)
+            return True
+
+        while stillHaveClones():
+            pass
+
+    def renameNodesAndKeepOriginals(self):
+        """
+        Sanitize the names of all nodes in the scene, returning a dict that points
+        renamed nodes to their original names so that they can be later restored.
+
+        Returns:
+            (dict of id to str): Dictionary with node ids and their original names
+        """
+        return dss.renameObjectsAndKeepOriginals(vrScenegraph.getAllNodes())
+
+    def renameMatsAndKeepOriginals(self):
+        """
+        Sanitize the names of all materials in the scene, returning a name that points
+        renamed materials to their original names so that they can be restored.
+
+        Returns:
+            (dict of id to str): Dictionary with material ids and their original names
+        """
+        return dss.renameObjectsAndKeepOriginals(vrMaterialPtr.getAllMaterials())
+
+    def renameNodesInVariantsFile(self, xmlRoot, nodeIDToOriginalNames, transVariants):
+        """
+        Uses the received dicts to update the variants xml tree in place with the names of nodes
+        that have been renamed. This includes NodeVariants, LightVariants and TransformVariants,
+        as well as NodeVariant states. Also updates VariantSet states to match.
+
+        Args:
+            xmlRoot (xml.etree.ElementTree.Element): Root of the variants XML tree
+            nodeIDToOriginalNames (dict of id to str): Dictionary with node ids and their original names
+            transVariants (dict): Keys are node names. Each value is a dict, where variant names are keys, and
+                                  values are float arrays describing the actual transform variants
+        """
+        # Invert map, but consider that there might be multiple nodes with the same origName
+        origNameToNodeIDs = defaultdict(list)
+        logging.debug('nodeToOriginalNames:')
+        for nodeID, origName in nodeIDToOriginalNames.iteritems():
+            logging.debug('\tnode: ' + str(vrNodePtr.vrNodePtr(nodeID).getName()) + ', origName: ' + str(origName))
+            origNameToNodeIDs[origName].append(nodeID)
+
+        # We don't need the defaultdict behaviour anymore, so lets prevent some mistakes
+        origNameToNodeIDs = dict(origNameToNodeIDs)
+
+        logging.debug('origNamesToNodes:')
+        for origName, nodeIDs in origNameToNodeIDs.iteritems():
+            logging.debug('\torigName: ' + str(origName) + ', nodes: ' + str([vrNodePtr.vrNodePtr(id).getName() for id in nodeIDs]))
+
+        # Ignore these states
+        commandNames = set(['!All', '!Next', '!Next (Loop)', '!None', '!Previous', '!Previous (Loop)', '!Enable', '!Disable', '!Toggle'])
+
+        # Rename light variant nodes
+        for xmlNode in xmlRoot.iterfind(".//LightVariant"):
+            lightNameVar = xmlNode.get('base')
+            logging.debug('Analyzing LightVariant ' + str(lightNameVar))
+            if lightNameVar in origNameToNodeIDs:
+                # Not much we can do since they don't have States for each child.
+                # "Getting the first node with the name" mirrors what VRED does anyway though
+                nodeID = origNameToNodeIDs[lightNameVar][0]
+                node = vrNodePtr.vrNodePtr(nodeID)
+                newName = node.getName()
+                logging.debug('\tRenaming ' + str(lightNameVar) + ' to ' + str(newName))
+                xmlNode.set('base', newName)
+
+        # Renames transform variants by comparing variant names if we must
+        # Note that we will sanitize the Transform variant names later, as these don't need any
+        # special connection with the fbx file. Both the name and the transform are in the xml
+        for xmlNode in xmlRoot.iterfind(".//TransformVariant"):
+            transNameVar = xmlNode.get('base')
+            logging.debug('Analyzing TransformVariant ' + str(transNameVar))
+
+            # The node of this transform variant hasn't been renamed
+            if transNameVar not in origNameToNodeIDs:
+                continue
+
+            renamedNodeIDList = origNameToNodeIDs[transNameVar]
+            numRenamedNodes = len(renamedNodeIDList)
+
+            if numRenamedNodes == 1:
+                node = vrNodePtr.vrNodePtr(renamedNodeIDList[0])
+                newName = node.getName()
+                logging.debug('\tonly one node renamed away from ' + str(transNameVar) + ': ' + str(newName))
+                node.set('base', newName)
+            elif numRenamedNodes > 1:
+                states = [stateNode.get('name') for stateNode in xmlNode.iterfind('State')]
+                states = [state for state in states if state not in commandNames]
+                states.sort()
+
+                logging.debug('\tmore than one node renamed away from ' + str(transNameVar))
+                logging.debug('\tstates: ' + str(states))
+
+                foundNode = False
+                for renamedNodeID in renamedNodeIDList:
+                    renamedNode = vrNodePtr.vrNodePtr(renamedNodeID)
+                    if renamedNode.getName() in transVariants:
+                        nodeTransVariants = transVariants[renamedNode.getName()]
+                        nodeTransVariants = nodeTransVariants.keys()
+                        nodeTransVariants.sort()
+
+                        logging.debug('\t\tcomparing against ' + str(renamedNode.getName()))
+                        logging.debug('\t\tchildren: ' + str(nodeTransVariants))
+
+                        if cmp(states, nodeTransVariants) == 0:
+                            logging.debug("\t\tthat's it! they both have the same children")
+                            foundNode = True
+                            newName = renamedNode.getName()
+                            xmlNode.set('base', newName)
+                            break
+
+                if not foundNode:
+                    logging.warning('Did not find a match for TransformVariant ' + str(transNameVar))
+
+        # Renames node variant switches by trying to identify them by name and children
+        # Note that we rename the 'base' attribute only. 'name' is the actual name of the variant
+        # and is sanitized separately, since it needs to match potential variant set references
+        renamedNodeVariants = {}
+        for xmlNode in xmlRoot.iterfind(".//NodeVariant"):
+            parentNameVar = xmlNode.get('base')
+            logging.debug('Analyzing NodeVariant ' + str(parentNameVar))
+
+            # The geometry switch node of this node variant hasn't been renamed
+            if parentNameVar not in origNameToNodeIDs:
+                continue
+
+            renamedNodeIDList = origNameToNodeIDs[parentNameVar]
+            numRenamedNodes = len(renamedNodeIDList)
+            # Only one node was renamed away from parentNameVar, so just rename node
+            # to its new name
+            if numRenamedNodes == 1:
+                node = vrNodePtr.vrNodePtr(renamedNodeIDList[0])
+                newName = node.getName()
+
+                if newName in renamedNodeVariants.values():
+                    logging.warning(str(newName) + ' was already used!')
+                renamedNodeVariants[parentNameVar] = newName
+
+                logging.debug('\tonly one node renamed away from ' + str(parentNameVar) + ': ' + str(newName))
+                xmlNode.set('base', newName)
+            # There are multiple switch nodes that were renamed from parentNameVar.
+            # Try to match as best as we can by comparing the node's children with
+            # the variant node's 'State' subnodes
+            # origNameToNodeIDs is a defaultdict, so it will always create an empty list
+            elif numRenamedNodes > 1:
+                states = [stateNode.get('name') for stateNode in xmlNode.iterfind('State')]
+                states = [state for state in states if state not in commandNames]
+                states.sort()
+
+                logging.debug('\tmore than one node renamed away from ' + str(parentNameVar))
+                logging.debug('\tstates: ' + str(states))
+
+                foundNode = False
+                for renamedNodeID in renamedNodeIDList:
+                    renamedNode = vrNodePtr.vrNodePtr(renamedNodeID)
+                    renamedNodeChildren = dsutils.getChildren(renamedNode)
+                    childrenOrigNames = []
+                    for child in renamedNodeChildren:
+                        if child.getID() in nodeIDToOriginalNames:
+                            childrenOrigNames.append(nodeIDToOriginalNames[child.getID()])
+                        else:
+                            childrenOrigNames.append(child.getName())
+
+                    logging.debug('\t\tcomparing against ' + str(renamedNode.getName()))
+                    logging.debug('\t\tchildren: ' + str(childrenOrigNames))
+
+                    childrenOrigNames.sort()
+                    if cmp(states, childrenOrigNames) == 0:
+                        logging.debug("\t\tthat's it! they both have the same children")
+                        foundNode = True
+                        newName = renamedNode.getName()
+
+                        if newName in renamedNodeVariants.values():
+                            logging.warning(str(newName) + ' was already used!')
+                        renamedNodeVariants[parentNameVar] = newName
+
+                        xmlNode.set('base', newName)
+                        break
+
+                if not foundNode:
+                    logging.warning('Did not find a match for NodeVariant ' + str(parentNameVar))
+
+        # Rename node variant states using the recently renamed parents
+        # Note that we couldn't do that in the above loop since we had no way of knowing which parent
+        # was which, but we could use direct node object references for the children there, which removes
+        # ambiguities. Now that we know the parent and they have unique names, we can easily use that to rename
+        # duplicate child states (especially if they are duplicate between different parents)
+        renamedNodeVariantStates = defaultdict(dict)
+        for xmlNode in xmlRoot.iterfind(".//NodeVariant"):
+            parentName = xmlNode.get('base')
+
+            # Pairs all 'Default' nodes with the 'ref' statenamekey and 'State' nodes with the 'name' statenamekey
+            for state, stateNameKey in chain(izip_longest(xmlNode.findall('Default'), [], fillvalue='ref'), \
+                                             izip_longest(xmlNode.findall('State'), [], fillvalue='name')):
+                stateName = state.get(stateNameKey)
+                logging.debug('Analyzing NodeVariant state ' + str(stateName) + ' from parent ' + str(parentName))
+
+                # The node this state refers to hasn't been renamed
+                if stateName not in origNameToNodeIDs:
+                    continue
+
+                renamedNodeIDList = origNameToNodeIDs[stateName]
+
+                logging.debug('\tnodes that were renamed away from ' + str(stateName) + ': ' + str([vrNodePtr.vrNodePtr(n).getName() for n in renamedNodeIDList]))
+
+                foundNode = False
+                for renamedNodeID in renamedNodeIDList:
+                    renamedNode = vrNodePtr.vrNodePtr(renamedNodeID)
+                    renamedParentNode = renamedNode.getParent().getName()
+                    logging.debug('\t' + str(renamedNode.getName()) + ' has parent ' + str(renamedParentNode))
+                    if renamedParentNode == parentName:
+                        foundNode = True
+                        newName = renamedNode.getName()
+                        logging.debug("\tthat's it, renaming " + str(stateName) + ' to ' + str(newName))
+
+                        nodeVariantName = xmlNode.get(stateNameKey)
+                        renamedNodeVariantStates[nodeVariantName][stateName] = newName
+
+                        state.set(stateNameKey, newName)
+                        break
+
+                if not foundNode:
+                    logging.warning('Did not find a match for NodeVariant state ' + str(stateName) + ' for NodeVariant ' + str(parentName))
+
+        # Rename references to node variants from variant sets
+        for xmlNode in xmlRoot.iterfind(".//NodeVariantRef"):
+            maybeRenamedRef = xmlNode.get('ref')
+            oldStateName = xmlNode.get('state')
+
+            logging.debug('Analyzing NodeVariantRef ref ' + str(maybeRenamedRef))
+
+            if maybeRenamedRef in renamedNodeVariantStates:
+                renamedStates = renamedNodeVariantStates[maybeRenamedRef]
+
+                if oldStateName in renamedStates:
+                    newStateName = renamedStates[oldStateName]
+                    logging.debug('\tstate ' + str(oldStateName) + ' seems to have been renamed to ' + str(newStateName))
+                    xmlNode.set('state', newStateName)
+
+    def renameClipsInVariantsFile(self, xmlRoot, clipOrigNamesToNewNames):
+        """
+        Uses the passed in dictionary to update references to a renamed animation clip in the variants xml tree
+
+        Args:
+            xmlRoot (xml.etree.ElementTree.Element): Root of the variants XML tree
+            clipOrigNamesToNewNames (dict of str to str): Dictionary with clips original names mapping to their new names
+        """
+        # In case of clips with duplicate names, our hands are tied since its impossible
+        # to get which clip the variant set used via anything other than its name string
+        for xmlNode in xmlRoot.iterfind('.//AnimationRef'):
+            oldRef = xmlNode.get('ref')
+            logging.debug('Analyzing AnimationRef ' + str(oldRef))
+            if oldRef in clipOrigNamesToNewNames:
+                newRef = clipOrigNamesToNewNames[oldRef]
+                logging.debug('\trenaming ' + str(oldRef) + ' to ' + str(newRef))
+                xmlNode.set('ref', newRef)
+
+    def renameMaterialsInVariantsFile(self, xmlRoot, matIDsToOriginalNames):
+        """
+        Uses the passed in dictionary to update references to a renamed material in the variants file
+
+        Args:
+            xmlRoot (xml.etree.ElementTree.Element): Root of the variants XML tree
+            matIDsToOriginalNames (dict of id to str): Dictionary with material ids and their original names
+        """
+        # Invert map, but consider that there might be multiple nodes with the same origName
+        origNameToMatIDs = defaultdict(list)
+        logging.debug('matIDsToOriginalNames:')
+        for matID, origName in matIDsToOriginalNames.iteritems():
+            logging.debug('\tmat: ' + str(vrMaterialPtr.vrMaterialPtr(matID).getName()) + ', origName: ' + str(origName))
+            origNameToMatIDs[origName].append(matID)
+
+        # We don't need the defaultdict behaviour anymore, so lets prevent some mistakes
+        origNameToMatIDs = dict(origNameToMatIDs)
+
+        logging.debug('origNamesToMatIDs:')
+        for origName, matIDs in origNameToMatIDs.iteritems():
+            logging.debug('\torigName: ' + str(origName) + ', mats: ' + str([vrMaterialPtr.vrMaterialPtr(matID).getName() for matID in matIDs]))
+
+        # Ignore these states
+        commandNames = set(['!All', '!Next', '!Next (Loop)', '!None', '!Previous', '!Previous (Loop)', '!Enable', '!Disable', '!Toggle'])
+
+        # Renames material variants by trying to identify them by name and submaterials
+        # Note that we rename the 'base' attribute only. 'name' is the display name of the variant
+        # and is sanitized separately, since it needs to match potential variant set references
+        renamedMatVariants = {}
+        for matXmlNode in xmlRoot.iterfind(".//MaterialVariant"):
+            switchMatNameVar = matXmlNode.get('base')
+            logging.debug('Analyzing MaterialVariant ' + str(switchMatNameVar))
+
+            # This switch material has not been renamed
+            if switchMatNameVar not in origNameToMatIDs:
+                continue
+
+            renamedMatIDList = origNameToMatIDs[switchMatNameVar]
+            numRenamedNodes = len(renamedMatIDList)
+            # Only one node was renamed away from switchMatNameVar, so just rename node
+            # to its new name
+            if numRenamedNodes == 1:
+                mat = vrMaterialPtr.vrMaterialPtr(renamedMatIDList[0])
+                newName = mat.getName()
+
+                if newName in renamedMatVariants.values():
+                    logging.warning(str(newName) + ' was already used!')
+                renamedMatVariants[switchMatNameVar] = newName
+
+                logging.debug('\tonly one mat renamed away from ' + str(switchMatNameVar) + ': ' + str(newName))
+                matXmlNode.set('base', newName)
+            # There are multiple materials that were renamed away from switchMatNameVar.
+            # Try to match as best as we can by comparing the mats's submaterials with
+            # the variant node's 'State' subnodes
+            elif numRenamedNodes > 1:
+                # We'll use sets for these comparisons because VRED gets confused if there are multiple
+                # submaterials with the same name under a switch material, only showing one in the variants
+                # window and in the emitted variants file.
+                states = [stateNode.get('name') for stateNode in matXmlNode.iterfind('State')]
+                states = [stateNode for stateNode in states if stateNode not in commandNames]
+                states = set(states)
+
+                logging.debug('\tmore than one mat renamed away from ' + str(switchMatNameVar))
+                logging.debug('\txml node states: ' + str(states))
+
+                foundMat = False
+                for renamedMatID in renamedMatIDList:
+                    renamedMat = vrMaterialPtr.vrMaterialPtr(renamedMatID)
+                    logging.debug('\t\tcomparing against ' + str(renamedMat.getName()))
+
+                    renamedMatSubmats = renamedMat.getSubMaterials()
+                    submatOrigNames = set()
+                    for submat in renamedMatSubmats:
+                        submatID = submat.getID()
+                        if submatID in matIDsToOriginalNames:
+                            submatOrigNames.add(matIDsToOriginalNames[submatID])
+                        else:
+                             submatOrigNames.add(submat.getName())
+                    logging.debug('\t\tsubmaterial names: ' + str(submatOrigNames))
+
+                    if states == submatOrigNames:
+                        logging.debug("\t\tthat's it! they both have the same (non-repeating) submaterials")
+                        foundMat = True
+                        newName = renamedMat.getName()
+
+                        if newName in renamedMatVariants.values():
+                            logging.warning(str(newName) + ' was already used!')
+                        renamedMatVariants[switchMatNameVar] = newName
+
+                        matXmlNode.set('base', newName)
+                        break
+
+                if not foundMat:
+                    logging.warning('Did not find a match for SwitchMaterial ' + str(switchMatNameVar))
+
+        # Rename switch material states using the recently renamed switch materials
+        renamedSwitchMaterialStates = defaultdict(dict)
+        for switchMatXmlNode in xmlRoot.iterfind(".//MaterialVariant"):
+            switchMatName = switchMatXmlNode.get('base')
+            logging.debug('Analyzing state nodes for switch material ' + str(switchMatName))
+
+            # We have already removed duplicate material names in the scene and renamed the switch material names
+            # in the xml hierarchy, so there should be only one material that has switchMatName
+            matPtrs = vrMaterialPtr.findMaterials(switchMatName)
+            if len(matPtrs) != 1:
+                logging.warning('Found ' + str(len(matPtrs)) + ' switch materials named ' + str(switchMatName) + '. Using the first if we can')
+
+            if len(matPtrs) == 0:
+                continue
+
+            submats = matPtrs[0].getSubMaterials()
+
+            # Find out how those submats were renamed, if they ever were
+            # Two cases:
+            #   Submaterials with the same name but only one appears in this switch material:
+            #       We'll cover this by comparing against actual submats of the switch material (which is presumed unique by now)
+            #   Submaterials with the same name and more than one appears in this switch material:
+            #       We need to have a state for each, so we need a list of tuples here, so that we can have multiple keys
+            #       pointing to the same value
+            submatOrigNamesToNewNames = []
+            for submat in submats:
+                submatID = submat.getID()
+                if submatID in matIDsToOriginalNames:
+                    submatOrigNamesToNewNames.append((matIDsToOriginalNames[submatID], submat.getName()))
+            logging.debug('\tsubmaterial old names to new names: ' + str(submatOrigNamesToNewNames))
+
+            # Rename default state
+            for defaultNode in switchMatXmlNode.iterfind("Default"):
+                stateName = defaultNode.get('ref')
+
+                # See if it has been renamed or not. We'll pick the first new name if we have
+                old2new = next((tup for tup in submatOrigNamesToNewNames if tup[0] == stateName), None)
+                if old2new is None:
+                    continue
+
+                newStateName = old2new[1]
+                defaultNode.set('ref', newStateName)
+                logging.debug('\t\trenamed default ' + str(stateName) + ' to ' + str(newStateName))
+
+            # Rename all other states
+            for stateNode in switchMatXmlNode.iterfind("State"):
+                stateName = stateNode.get('name')
+
+                old2new = next((tup for tup in submatOrigNamesToNewNames if tup[0] == stateName), None)
+                if old2new is None:
+                    continue
+
+                # Cover case #2, allowing us to pick the other newName whenever we search for this stateName again
+                submatOrigNamesToNewNames.remove(old2new)
+
+                logging.debug('\t\tsubmatsOrigNamesToNewNames now has ' + str(len(submatOrigNamesToNewNames)) + ' tuples')
+
+                newStateName = old2new[1]
+
+                # Keep track of renamed states so that we can update switch material references in variant sets if we need to
+                switchMatDisplayName = switchMatXmlNode.get('name')  # Variant sets reference variants by name, not node
+                renamedSwitchMaterialStates[switchMatDisplayName][stateName] = newStateName
+
+                stateNode.set('name', newStateName)
+                logging.debug('\t\trenamed ' + str(stateName) + ' to ' + str(newStateName))
+
+        # Renamed switch material references from variant sets
+        for matRefXmlNode in xmlRoot.iterfind(".//MaterialVariantRef"):
+            maybeRenamedRef = matRefXmlNode.get('ref')
+            oldStateName = matRefXmlNode.get('state')
+            logging.debug('Analyzing MaterialVariantRef ref ' + str(maybeRenamedRef))
+
+            if maybeRenamedRef in renamedSwitchMaterialStates:
+                renamedStates = renamedSwitchMaterialStates[maybeRenamedRef]
+                if oldStateName in renamedStates:
+                    newStateName = renamedStates[oldStateName]
+                    logging.debug('\tstate ' + str(oldStateName) + ' seems to have been renamed to ' + str(newStateName))
+                    matRefXmlNode.set('state', newStateName)
+
+    def restoreOriginalNames(self, objToOriginalNames):
+        """
+        Receives a dictionary from objects to strings of their original names.
+        Calls key.setName(value) on every item of the dictionary
+
+        Args:
+            objToOriginalNames (dict of obj to str): Dictionary with objects and their original names
+        """
+        for pair in objToOriginalNames.items():
+            pair[0].setName(pair[1])
+
+    def getSafeVREDVersion(self):
+        """
+        Returns a floating point number indicating the current VRED version, or zero if something fails
+
+        Returns:
+            (float): floating point with current major and minor VRED versions
+        """
+        try:
+            return dss.sanitize(str(vrController.getVredVersion()))[0]
+        except:
+            return "Unknown"
+
+    def getUnitScale(self):
+        """
+        Returns the unit scale of the scene. 1.0 = mm, 10.0 = cm, 1000.0 = m
+
+        Returns:
+            (float): Value of the unit scale
+        """
+        # Since there is no API for getting this, we'll dig through the GUI for a combobox
+        # in the status bar that kind of looks like what we want.
+        # The combobox has no objectName, and although we could try filtering by the toolTip,
+        # that is likely to change without warning, while these units are not
+        units = {'mm': 1.0, 'cm': 10.0, 'm': 1000.0}
+
+        try:
+            for item in self.mainWindow.statusBar().findChildren(QtWidgets.QComboBox):
+                currText = item.currentText()
+                if currText in units.keys():
+                    return units[currText]
+        except AttributeError:
+            pass
+
+        return 1.0
+
+    def getUnitScaleFactorTocm(self):
+        """
+        Returns the multiplicative factor to convert whatever scale is currently selected
+        to cm, which will be the units of the exported .fbx file
+        """
+        return 10.0 / self.getUnitScale()
+
+    def removeNURBS(self):
+        """
+        Travels the entire hierarchy, finds NURBS nodes and converts them. Also temporarily stores
+        the node's transform in a sibling node and restores it back after conversion, since VRED doesn't
+        maintain it during the conversion for some reason
+        """
+        def removeNURBSRecursive(node, nurbsNodes):
+            for index in range(node.getNChildren()):
+                child = node.getChild(index)
+
+                # Nodes with this attachment *may* be NURBs, but all NURBs have this attachment
+                if dsutils.isNURBS(child):
+                    nurbsNodes.append(child)
+                else:
+                    removeNURBSRecursive(child, nurbsNodes)
+
+        nodesToFix = []
+        removeNURBSRecursive(vrScenegraph.getRootNode(), nodesToFix)
+
+        logging.debug('Removing NURBS from nodes ' + str([a.getName() for a in nodesToFix]))
+
+        count = 0
+        for node in nodesToFix:
+            nodeOrigName = node.getName()
+            parent = node.getParent()
+
+            newNode = vrScenegraph.createNode('Transform3D', str(nodeOrigName) + "_Temp", parent)
+
+            vrScenegraph.copyTransformation(node, newNode)
+
+            vrOptimize.removeNURBS(node)
+            count += 1
+
+            vrScenegraph.copyTransformation(newNode, node)
+
+            newNode.sub()
+
+        return count
+
+    def isSceneDirty(self):
+        titleText = ' - Autodesk'
+        windowTitle = self.mainWindow.windowTitle()
+
+        try:
+            starPos = windowTitle.rfind(titleText)
+            if starPos != -1:
+                if windowTitle[starPos-1] != '*':
+                    return False
+        except:
+            pass
+
+        return True
+
+    def exportSceneDialog(self):
+        """
+        Called when File->Export->Export to Datasmith... is called. Exports the scene FBX as well
+        as auxilliary files like the variants *.var file and the lights *.lights file, all with the
+        same name as the FBX file, exported on the same folder
+        """
+        # Open the file dialog to export the scene
+        filename = vrFileDialog.getSaveFileName("Export scene", "", ["fbx(*.fbx)"], True)
+        filename = str(QtCore.QDir.toNativeSeparators(filename))
+        self.exportScene(filename)
+
+    def exportScene(self, filename, isTest=False):
+        if filename:
+            filenameNoExt = os.path.splitext(filename)[0]
+            originalMatNames = {}
+            filenameTemp = ""
+            vred2UnrealNode = None
+            exporterVersionNode = None
+            savedvpbPath = ""
+            exportCompleted = False
+            userCanceled = False
+            cameraTransform = None
+
+            try:
+                if not isTest and self.isSceneDirty():
+                    msgBox = QMessageBox()
+                    msgBox.setText('The scene has unsaved changes, so it will be saved before exporting to Datasmith. Do you wish to proceed?\n')
+                    msgBox.setDetailedText('In order to extract all scene information with maximum fidelity, some operations are required which, due to limitations in the VRED API, are undoable or destructive to the scene.\n\nThe export process will automatically save, perform those changes, then reload the scene.')
+                    msgBox.setStandardButtons(QMessageBox.Yes | QMessageBox.No)
+                    msgBox.setDefaultButton(QMessageBox.No)
+                    msgBox.setIcon(QMessageBox.Icon.Warning)
+                    chosenOption = msgBox.exec_()
+                    if chosenOption != QMessageBox.Yes:
+                        userCanceled = True
+                        return
+
+                logging.info('=========================================')
+                logging.info('||      Starting Datasmith export      ||')
+                logging.info('=========================================')
+                logging.info('Target file: ' + str(filename))
+                logging.info('Exporter version: ' + str(exporterVersion))
+                logging.info('VRED version: ' + str(self.getSafeVREDVersion()).replace('_', '.'))
+
+                # Preserve camera position before/after export as loading will
+                # reset it
+                cameraTransform = vrCamera.getActiveCameraNode().getWorldTransform()
+
+                # disable rendering to speed up things
+                vrOSGWidget.enableRender(False)
+                # disable scenegraph update to speed up things
+                vrScenegraph.enableScenegraph(False)
+
+                # Save scene since some operations are destructive and there is no workaround around some
+                # API limitations
+                if not isTest and self.isSceneDirty():
+                    logging.info('Saving temp scene file...')
+                    savedvpbPath, saveWorked = dsutils.save()
+                    if not saveWorked:
+                        return
+                if len(savedvpbPath) == 0:
+                    savedvpbPath = vrFileIO.getFileIOFilePath()
+
+                # Clear undo stack to free deleted nodes that haven't been actually deleted yet
+                vrController.clearUndoStack()
+
+                settings = UserSettings()
+                settings.tryFetchFromFile()
+
+                # Remove cloned and transformable root nodes
+                logging.info('Unsharing nodes...')
+                self.unshareNodes()
+
+                # Removing invalid nodes
+                logging.info('Removing invalid nodes...')
+                vrOptimize.removeEmptyGeometries(vrScenegraph.getRootNode())
+                vrOptimize.removeInvalidTexCoords(vrScenegraph.getRootNode())
+
+                # Convert NURBs to meshes
+                logging.info('Removing NURBS...')
+                numRemoved = self.removeNURBS()
+                logging.info('Removed ' + str(numRemoved) + ' NURBS. Remaining NURBS: ' + str(dsutils.getNumSceneNURBs()))
+
+                # Sanitize node names
+                logging.info('Sanitizing node names...')
+                originalNodeNames = self.renameNodesAndKeepOriginals()
+
+                # Build an animation library so that we can keep track of clips and blocks
+                logging.info('Creating animation library...')
+                animLib = dsanim.catalogSceneAnimations(self.getUnitScaleFactorTocm(), settings.animBaseTime, settings.animPlayRate)
+                animLib.sanitize()
+                with open(filenameNoExt + '.clips', 'w') as animFile:
+                    animLib.serialize(animFile)
+
+                # Emit empty nodes to signal metadata
+                exporterVersionStr = str(exporterVersion).replace('.', '_')
+                vredVersionStr = self.getSafeVREDVersion()
+                logging.info('Creating metadata hierarchy...')
+                vred2UnrealNode = vrScenegraph.createNode('Transform', 'VRED2Unreal', vrScenegraph.getRootNode())
+                exporterVersionNode = vrScenegraph.createNode('Transform', exporterVersionStr, vred2UnrealNode)
+                vrScenegraph.createNode('Transform', vredVersionStr, exporterVersionNode)
+
+                # Extract transform variants from nodes that have them
+                logging.info('Extracting transform variants...')
+                transVariants = self.extractTransformVariants()
+                logging.debug('\tExtracted transform variants: ' + str(transVariants))
+
+                # Save variants file with a tmp extension
+                logging.info('Saving unprocessed variants file...')
+                filenameTemp = filenameNoExt + '.tmp'
+                vrVariantSets.saveVariants(filenameTemp)
+
+                # Read the variants file back so that we can post-process it
+                logging.info('Reading back unprocessed variants file...')
+                varfileXML = ET.parse(filenameTemp)
+                xmlRoot = varfileXML.getroot()
+
+                # Post-process VAR file
+                logging.info('Pre-processing variants file...')
+                self.removePreviewTags(xmlRoot)
+                self.renameNodesInVariantsFile(xmlRoot, originalNodeNames, transVariants)
+                if animLib is not None:
+                    self.renameClipsInVariantsFile(xmlRoot, animLib.ClipOriginalNamesToNewNames)
+
+                # Post-process materials in variants file
+                logging.info('Injecting switch material users...')
+                originalMatNames = self.renameMatsAndKeepOriginals()
+                self.renameMaterialsInVariantsFile(xmlRoot, originalMatNames)
+                self.injectSwitchMaterialUsers(xmlRoot)
+
+                # Post-process variants file
+                self.sanitizeVariantNamesAndVariantSetRefs(xmlRoot)
+                self.injectTransformVariants(xmlRoot, transVariants)
+                self.sanitizeTransformVariantStateNames(xmlRoot, transVariants)
+
+                # Write final var file
+                logging.info("Writing variants file '" + str(filenameNoExt) + ".var'...")
+                varfileXML.write(filenameNoExt + '.var')
+
+                # Save lights to lights file
+                logging.info("Writing lights file '" + str(filenameNoExt) + ".lights'...")
+                lights = self.extractLights()
+                self.writeLightsFile(filenameNoExt + '.lights', lights)
+
+                # Write textures
+                logging.info("Extracting material data...")
+                mats = dsmat.extractMaterials(vrMaterialPtr.getAllMaterials())
+                rootFolder = os.path.split(filename)[0]
+                texFolder = os.path.join(rootFolder, 'Textures')
+                try:
+                    os.mkdir(texFolder)
+                except OSError:  # Dir already exists
+                    pass
+                self.writeTextures(texFolder, mats)
+
+                # Save materials to mat file
+                logging.info("Writing mat file '" + str(filenameNoExt) + ".mat'...")
+                self.writeMatFile(filenameNoExt + '.mats', mats)
+
+                # Preserve unused Switch Material variants
+                # After injecting switch material users so temp nodes don't show
+                # up on the var file
+                # This is useless if we're exporting with mats file, but will preserve these
+                # materials in the FBX file if we're importing without them, for some reason
+                logging.info('Preserving unused switch material variants...')
+                matsToPreserve = self.getMaterialsToPreserve()
+                self.assignMatsToTempGeometry(matsToPreserve)
+
+                # Save FBX scene
+                logging.info('Saving FBX scene...')
+                vrFileIO.save(filename)
+
+                exportCompleted = True
+
+            except Exception as e:
+                logging.error(e)
+
+            finally:
+                if userCanceled:
+                    return False
+
+                if os.path.exists(filenameTemp):
+                    os.remove(filenameTemp)
+
+                vrController.newScene()
+                if not isTest and len(savedvpbPath) > 0 and os.path.exists(savedvpbPath):
+                    logging.info('Loading back saved scene...')
+                    vrFileIO.load([savedvpbPath], vrScenegraph.getRootNode().getParent(), True, False)
+
+                # Restore camera position to what it was before the export
+                if cameraTransform is not None:
+                    vrCamera.getActiveCameraNode().setTransformMatrix(cameraTransform, False)
+
+                vrScenegraph.enableScenegraph(True) # reenable scenegraph updates
+                vrOSGWidget.enableRender(True) # reenable rendering
+
+                if exportCompleted:
+                    logging.info('=========================================')
+                    logging.info('||          Export  succeeded          ||')
+                    logging.info('=========================================')
+                else:
+                    logging.info('=========================================')
+                    logging.info('||           Export   failed           ||')
+                    logging.info('=========================================')
+                logging.info('Target file: ' + str(filename))
+                logging.info('Exporter version: ' + str(exporterVersion))
+                logging.info('VRED version: ' + str(self.getSafeVREDVersion()).replace('_', '.'))
+
+                return exportCompleted
diff --git a/Engine/Plugins/Enterprise/DatasmithFBXImporter/Resources/VREDPlugin/DatasmithExporter/log.py b/Engine/Plugins/Enterprise/DatasmithFBXImporter/Resources/VREDPlugin/DatasmithExporter/log.py
new file mode 100644
--- /dev/null
+++ b/Engine/Plugins/Enterprise/DatasmithFBXImporter/Resources/VREDPlugin/DatasmithExporter/log.py
@@ -0,0 +1,8 @@
+import logging
+import sys
+
+FORMAT = '[%(asctime)s][%(levelname)8s] DatasmithExporter: %(message)s'
+DATE_FORMAT = '%H:%M:%S'
+logging.basicConfig(stream=sys.stdout, level=logging.INFO, format=FORMAT, datefmt=DATE_FORMAT)
+
+#logging.DEBUG logging.INFO
\ No newline at end of file
diff --git a/Engine/Plugins/Enterprise/DatasmithFBXImporter/Resources/VREDPlugin/DatasmithExporter/main.py b/Engine/Plugins/Enterprise/DatasmithFBXImporter/Resources/VREDPlugin/DatasmithExporter/main.py
new file mode 100644
--- /dev/null
+++ b/Engine/Plugins/Enterprise/DatasmithFBXImporter/Resources/VREDPlugin/DatasmithExporter/main.py
@@ -0,0 +1,80 @@
+from PySide2 import QtCore, QtGui, QtWidgets
+
+import sys
+import os
+import logging
+import xml.etree.ElementTree as ET
+from collections import defaultdict
+from itertools import izip_longest, chain
+
+import vrFileIO
+import vrFileDialog
+import vrVariantSets
+import vrMaterialPtr
+import vrScenegraph
+import vrNodeUtils
+import vrFieldAccess
+import vrController
+import vrOptimize
+import vrNodePtr
+
+# Adds the ScriptPlugins folder to path so that we can find the other files
+# in our module
+scriptPluginsFolder = os.path.dirname(os.path.realpath(os.curdir))
+if scriptPluginsFolder not in sys.path:
+    sys.path.append(scriptPluginsFolder)
+
+from DatasmithExporter.exporter import DatasmithFBXExporter
+from DatasmithExporter.exporter import exporterVersion
+
+def vredMainWindow(id):
+    from shiboken2 import wrapInstance
+    return wrapInstance(id, QtWidgets.QMainWindow)
+
+class DatasmithExporterPlugin():
+    def __init__(self, mw):
+        # Store this since it's only globally available when the plugin first loads
+        self.mainWindow = mw
+        mb = self.mainWindow.menuBar()
+        actions = mb.actions()
+
+        self.exporter = DatasmithFBXExporter(mw)
+
+        # Create an action
+        self.exportSceneAction = QtWidgets.QAction(mw.tr("Export to Datasmith..."), mw)
+        self.exportSceneAction.triggered.connect(self.exporter.exportSceneDialog)
+
+        # Iterate over the main menu bar
+        for action in actions:
+
+            # Iterate over the file menu
+            if action.text() == mw.tr("&File"):
+                fileMenuItems = action.menu().actions()
+                for fileItem in fileMenuItems:
+
+                    # Add our action after a separator at the end of the Export menu
+                    if fileItem.text() == mw.tr("Export Scene Data") or fileItem.text() == mw.tr("Export"):
+                        exportMenuItems = fileItem.menu().actions()
+
+                        # Get indices of all entries that have "Datasmith" in them
+                        indices = [pair[0] for pair in enumerate(exportMenuItems) if "Export to Datasmith" in pair[1].text()]
+
+                        # Append indices of the previous elements too as we always add
+                        # a separator before our entries
+                        for index in reversed(indices):
+                            indices.append(index-1)
+
+                        # Remove old entries as we need new ones to point to our reloaded
+                        # modules and classes
+                        for index in reversed(sorted(indices)):
+                            itemToDelete = exportMenuItems[index]
+                            if "Export to Datasmith" in itemToDelete.text() or itemToDelete.isSeparator():
+                                fileItem.menu().removeAction(itemToDelete)
+
+                        fileItem.menu().addSeparator()
+                        fileItem.menu().addAction(self.exportSceneAction)
+
+        logging.info('Finished loading DatasmithExporter plugin version ' + str(exporterVersion))
+
+exporter = DatasmithExporterPlugin(vredMainWindow(VREDMainWindowId))
+
diff --git a/Engine/Plugins/Enterprise/DatasmithFBXImporter/Resources/VREDPlugin/DatasmithExporter/mat.py b/Engine/Plugins/Enterprise/DatasmithFBXImporter/Resources/VREDPlugin/DatasmithExporter/mat.py
new file mode 100644
--- /dev/null
+++ b/Engine/Plugins/Enterprise/DatasmithFBXImporter/Resources/VREDPlugin/DatasmithExporter/mat.py
@@ -0,0 +1,597 @@
+import sys
+import os
+import logging
+
+import vrScenegraph
+import vrFieldAccess
+import vrNodePtr
+import vrNodeUtils
+import vrMaterialPtr
+
+# Adds the ScriptPlugins folder to path so that we can find the other files
+# in our module
+scriptPluginsFolder = os.path.dirname(os.path.realpath(os.curdir))
+if scriptPluginsFolder not in sys.path:
+    sys.path.append(scriptPluginsFolder)
+
+import DatasmithExporter.utils as dsutils
+
+class Tex:
+    def __init__(self):
+        self.srcPath = ''  # Path stored in VRED, used by the artist
+        self.exportedPath = ''  # Path of the exported pixel data during export process
+        self.texType = ''
+        self.pixelData = []
+        self.width = -1
+        self.height = -1
+        self.bytesPerPixel = -1
+        self.channels = -1
+        self.stride = 0  # Corrected stride, since we need to have pixels and rows aligned to 32-bit
+        self.pixelFormat = -1
+        self.dataType = -1  # unsigned char, etc
+        self.properties = {}
+
+    def __repr__(self):
+        text = '\tTex at path "' + str(self.srcPath) + '", Type: ' + str(self.texType) + '\n'
+        text += '\tProperties:\n'
+
+        text += '\t\twidth: ' + str(self.width) + '\n'
+        text += '\t\theight: ' + str(self.height) + '\n'
+        text += '\t\tbytesPerPixel: ' + str(self.bytesPerPixel) + '\n'
+        text += '\t\tchannels: ' + str(self.channels) + '\n'
+        text += '\t\tstride: ' + str(self.stride) + '\n'
+        text += '\t\tpixelFormat: ' + str(self.pixelFormat) + '\n'
+        text += '\t\tdataType: ' + str(self.dataType) + '\n'
+
+        for prop in self.properties:
+            text += '\t\t' + str(prop) + ': ' + str(self.properties[prop]) + '\n'
+
+        return text
+
+class Mat:
+    def __init__(self):
+        self.name = ""
+        self.matType = ""
+        self.vredMatId = 0
+        self.users = []  # Name of nodes that use this material
+        self.properties = {}  # indexed by name e.g. diffuseColor, specularColor, roughness
+        self.textures = {}  # indexed by name e.g. diffuseTexture, roughnessTexture
+        self.submaterials = []  # name of submaterials (if we're a composite material like a switch material)
+
+    def __repr__(self):
+        text = 'Mat "' + str(self.name) + '", Type: ' + str(self.matType) + ', VREDID: ' + str(self.vredMatId) + '\nUsers:\n'
+
+        for user in self.users:
+            text += '\t' + str(user) + '\n'
+
+        text += 'Properties:\n'
+
+        for prop in self.properties:
+            text += '\t' + str(prop) + ': ' + str(self.properties[prop]) + '\n'
+
+        text += 'Textures:\n'
+
+        for tex in self.textures:
+            text += '\t' + str(tex) + ': ' + str(self.textures[tex]) + '\n'
+
+        text += '\n'
+        return text
+
+# mat.fields().hasField('diffuseColor')
+fieldsPropList = {
+    'diffuseColor':'SFVec3f',
+    'specularColor':'SFVec3f',
+    'occlusionColor':'SFVec3f',  # Common
+    'occlusionIntensity':'SFReal32',  # Common
+    'lightingMode':'SFInt8',  # Common
+    'useStructure':'SFBool',  # Roughness texture
+    'parallaxIntensity':'SFReal32',  # Roughness texture
+    'bumpIntensity':'SFReal32',  # Roughness texture
+    'structureSize':'SFReal32',  # Roughness texture
+    'structureIntensity':'SFReal32',  # Roughness texture
+    #'bumpType':'SFInt32',  # Roughness texture
+    'incandescenceColor':'SFVec4f',  # Incandescence texture, intensity is alpha
+    'roughness':'SFReal32',  # Plastic, flake roughness
+    'reflectivity':'SFReal32',  # Plastic, metallic carpaint, brushed metal
+    'seeThrough':'SFVec3f',  # 1 - alpha
+    'diffuseBackscattering':'SFReal32',  # Phong
+    'fresnelQuality':'SFUInt32',  # Reflective Plastic
+    'reflectionColor':'SFVec3f',  # Chrome
+    'smear':'SFReal32',  # Chrome
+    'contrast':'SFReal32',  # Chrome
+    'saturation':'SFReal32',  # Chrome
+    'metalType':'SFInt32',  # Chrome
+    'isRough':'SFBool',  # Chrome
+    'roughnessU':'SFReal32',  # Brushed metal
+    'roughnessV':'SFReal32',  # Brushed metal
+    'exteriorColor':'SFVec3f',  # Glass, value in HSV is its alpha (??)
+    'interiorColor':'SFVec3f',  # Glass, value in HSV is its alpha (??)
+    'useFrostedGlass':'SFBool',  # Glass, useRoughness in UI
+    'useReflectivity':'SFBool',  # Glass, Use Reflectivity in UI
+    'refractionIndex':'SFReal32',  # Glass, IOR
+    'paintType':'SFUInt32',  # Metallic carpaint
+    'baseColor':'SFVec3f',  # Metallic carpaint
+    'flakeColor':'SFVec3f',  # Metallic carpaint
+    'useFlipFlop':'SFBool',  # Metallic carpaint
+    'flipFlopFlakeColor':'SFVec3f',  # Metallic carpaint
+    'flipFlopBlending':'SFReal32',  # Metallic carpaint
+    'flakeReflectivity':'SFReal32',  # Metallic carpaint
+    'flakeSize':'SFReal32',  # Metallic carpaint, flipflop carpaint, 1/ui value
+    'flakeIntensity':'SFReal32',  # Metallic carpaint, Flake Perturbation in ui, flipflop carpaint
+    'clearcoatColor':'SFVec3f',  # Metallic carpaint
+    'clearcoatThickness':'SFReal32',  # Metallic carpaint
+    'clearcoatDensity':'SFReal32',  # Metallic carpaint
+    'clearcoatRefractionIndex':'SFReal32',  # Metallic carpaint
+    'useOrangePeel':'SFBool',  # Metallic carpaint
+    'orangePeelFrequency':'SFReal32',  # Metallic carpaint, wavy bump map
+    'orangePeelIntensity':'SFReal32',  # Metallic carpaint
+    'pigmentColor':'SFVec3f',  # Metallic carpaint
+    'pigmentDensity':'SFReal32',  # Metallic carpaint, Pigment concentration in UI
+    'flakeDensity':'SFReal32',  # Metallic carpaint, Flake Density in ui
+    'flakeColor1':'SFVec3f',  # Flipflop carpaint
+    'flakeColor2':'SFVec3f',  # Flipflop carpaint
+    'blending':'SFReal32',  # Flipflop carpaint
+    'brushAxis':'SFInt8',  # Brushed metal
+    'triplanarRotate':'SFVec3f',  # Brushed metal
+}
+
+# vrFieldAccess(mat.fields().getFieldContainer('colorComponentData')).hasField('iorType')
+ccdPropList = {
+    'clearcoatReflectivity':'SFReal32',  # Brushed metal
+    'displacementHeight':'SFReal32',  # Disp texture
+    'displacementOffset':'SFReal32',  # Disp texture, value added to disp
+    'displacementSilhouetteMode':'SFUInt32',  # Disp texture
+    'scatterType':'SFInt8',  # Subsurface scattering
+    'scatterColor':'SFVec3f',  # Subsurface scattering
+    'roughnessTextureRange':'SFVec4f',  # Roughness texture, min is [0], max is [2]
+    'iorType':'SFUInt32',  # Glass
+    'clearcoatType':'SFUInt32'  # Metallic carpaint
+}
+
+# vrFieldAccess(mat.fields().getFieldContainer('secondFlakeLayerSettings')).hasField('flakeColor')
+flakeLayerPropList = {
+    'flakeColor':'SFVec3f',  # Metallic carpaint
+    'flakeSize':'SFReal32',  # Metallic carpaint, flipflop carpaint, 1/ui value
+    'flakeRoughness':'SFReal32',  # Metallic carpaint, flipflop carpaint, 1/ui value
+    'flakeIntensity':'SFReal32',  # Metallic carpaint, Flake Perturbation in ui, flipflop carpaint
+    'flakeDensity':'SFReal32',  # Metallic carpaint, Flake Density in ui
+    'flakeReflectivity':'SFReal32',  # Metallic carpaint
+    'pigmentColor':'SFVec3f',  # Metallic carpaint
+    'pigmentDensity':'SFReal32',  # Metallic carpaint, Pigment concentration in UI
+}
+
+# vrFieldAccess(colorComponentData.getFieldContainer('diffuseComponent'))
+texTypes = [
+    'diffuseComponent',
+    'glossyComponent',
+    'specularComponent',
+    'incandescenceComponent',
+    'bumpComponent',
+    'transparencyComponent',
+    'scatterComponent',
+    'roughnessComponent',
+    'displacementComponent',
+    'fresnelComponent',
+    'rotationComponent',
+    'indexOfRefractionComponent',
+    'specularBumpComponent'
+]
+
+# vrFieldAccess(colorComponentData.getFieldContainer('diffuseComponent')).getReal32('anisotropy')
+texPropList = {
+    #'useTexture':'SFBool', We check this explicitly
+    'repeatMode':'SFUInt32',
+    'repeat':'SFVec2f',
+    'offset':'SFVec2f',
+    'rotate':'SFReal32',
+    'anisotropy':'SFReal32',
+    'gamma':'SFVec3f',
+    'invertTexture':'SFBool',
+    'textureTarget':'SFUInt32',
+    'color':'SFVec3f',
+    'connectRepeatModes':'SFBool',
+    'useImageSequence':'SFBool',
+    'imageNumber':'SFUInt32',
+    'frameOffset':'SFUInt32',
+    'componentType':'SFString',
+    'hasLinkedSettings':'SFBool',
+    'textureProjectionType':'SFUInt32',
+    'planarProjectionFlags':'SFUInt32',
+    'triplanarRepeatU':'SFVec3f',
+    'triplanarRepeatV':'SFVec3f',
+    'triplanarOffsetU':'SFVec3f',
+    'triplanarOffsetV':'SFVec3f',
+    'triplanarBlend':'SFReal32',
+    'triplanarRotation':'SFVec3f',
+    'textureSizeMM':'SFVec2f',
+    'useRealWorldScale':'SFBool'
+}
+
+def getPropValue(fieldAccessObj, propName, propType):
+    if propType == 'SFReal32':
+        return fieldAccessObj.getReal32(propName)
+    elif propType == 'SFUInt32':
+        return fieldAccessObj.getUInt32(propName)
+    elif propType == 'SFBool':
+        return fieldAccessObj.getBool(propName)
+    elif propType == 'SFVec2f':
+        return fieldAccessObj.getVec(propName, 2)
+    elif propType == 'SFVec3f':
+        return fieldAccessObj.getVec(propName, 3)
+    elif propType == 'SFVec4f':
+        return fieldAccessObj.getVec(propName, 4)
+    elif propType == 'SFInt8':
+        return fieldAccessObj.getInt8(propName)
+    elif propType == 'SFUInt8':
+        return fieldAccessObj.getUInt8(propName)
+    elif propType == 'SFInt32':
+        return fieldAccessObj.getInt32(propName)
+    elif propType == 'SFString':
+        return fieldAccessObj.getString(propName)
+    else:
+        logging.error('Invalid property type ' + str(propType))
+        return None
+
+def unProcessTexture(tex):
+    """
+    Reverts some processing that VRED applies to bump and displacement maps
+    so that we can export an image that is as close as possible to whatever was originally
+    imported in VRED. This might cause a 32-bit pixel/row-aligned image to become unaligned!
+    """
+    logging.debug('\tUnprocessing texture ' + str(tex.srcPath))
+
+    if tex.texType == 'bumpComponent':
+        # There are many different processes that VRED does to a bump map when it is first imported:
+        # I've seen it discard the B channel and flip R and G; Do nothing; Flip R and G and do some
+        # kind of histogram normalization on B.
+        # Additionally it will, internally, create an RGB texture for the bump map regardless of
+        # whether it was originally a grayscale or an RGB image.
+        # In the end, we have no idea what has been done to this image when we have access to it.
+        # 'nothing' seems to be the majority of cases here, so we'll ignore them for now.
+        # Besides, we have to consider that this is a regular normal map on import, so the user
+        # can quickly switch to another source image on UE4 side and things will just work
+        return
+
+    elif tex.texType == 'displacementComponent':
+        # If this was an RGB image on import, VRED seems to convert it to grayscale.
+        # The stored texture will be RGB regardless though, but the grayscale image will be stored on the B
+        # channel only. R and G channels seems to store some kind of processed version (erode/dilate?)
+        if tex.channels == 3 and tex.bytesPerPixel == 3 and tex.dataType == 5121:
+            newBpp = 1
+            newChannels = 1
+            newStride = tex.width * newBpp
+
+            newPixelData = [0]*(newBpp * tex.width * tex.height)
+            for i in range(tex.height):
+                for j in range(tex.width):
+                    newPixelData[i * newStride + j] = tex.pixelData[i * tex.stride + j * tex.bytesPerPixel + 2]
+
+            tex.pixelData = newPixelData
+            tex.bytesPerPixel = 1
+            tex.channels = newChannels
+            tex.stride = newStride
+            return
+
+        logging.warning('Unknown problem unProcessing this displacement texture: ' + str(tex))
+
+def prepareTexturePixelData(tex):
+    """
+    Aligns each pixel (and each row) to 32-bit boundaries so that it can be exported with PySide2's PNG
+    exporter. It will ignore strange formats (like non-uchar images and so on) as those will be exported
+    as raw binary data anyway.
+
+    Due to the padding required and supported export image formats, we only have access to
+    8-bit 1-channel and 8-bit 4-channel, really.
+
+    Also flips the image vertically
+    """
+    logging.debug('\tPreparing texture pixel data ' + str(tex.srcPath))
+
+    # someVals = allMips[:50]
+    # logging.warning('Tex ' + str(tex.srcPath) + ' was originally: \n\ttype: ' + str(tex.texType) + '\n\tbpp: ' + str(tex.bytesPerPixel) + '\n\tchannels: ' + str(tex.channels) + '\n\tpixFormat: ' + str(tex.pixelFormat) + '\n\tsome vals: ' + str(someVals))
+
+    validDataType = tex.dataType == 5121  # unsigned byte
+
+    # Realign RGB24 images to RGBA32 so they can be exported via PNG
+    # PySide2 only handles up to RGBA32 (no RGBA64),
+    # meaning there is no point trying to fixup bpps of 5 or 7, or other channel types, if those even exist
+    if validDataType and tex.bytesPerPixel == 3 and tex.channels == 3:
+        newBpp = 4
+        newChannels = 4
+
+        logging.debug('\t\tRealigning RGB24 image...')
+
+        # Pad rows to 32-bit if necessary
+        newStride = tex.width * newBpp
+        remainder = newStride % 4
+        if remainder != 0:
+            newStride = newStride + remainder
+        zeroMip = [0]*(newStride*tex.height)
+
+        for i in range(tex.height):
+            for j in range(tex.width):
+                zeroMip[i*newStride + j*newBpp + 0] = tex.pixelData[(tex.height-1-i)*tex.stride + j*tex.bytesPerPixel + 0]
+                zeroMip[i*newStride + j*newBpp + 1] = tex.pixelData[(tex.height-1-i)*tex.stride + j*tex.bytesPerPixel + 1]
+                zeroMip[i*newStride + j*newBpp + 2] = tex.pixelData[(tex.height-1-i)*tex.stride + j*tex.bytesPerPixel + 2]
+                zeroMip[i*newStride + j*newBpp + 3] = 255
+
+        tex.bytesPerPixel = newBpp
+        tex.channels = newChannels
+        tex.stride = newStride
+        tex.pixelData = zeroMip
+        return
+
+    # Like GL_LUMINANCE_ALPHA: Need to replicate first channel to RGB and put the second
+    # channel on alpha. This because the Qt exporter doesn't have support for any 16 bit
+    # 2-channel image
+    elif validDataType and tex.bytesPerPixel == 2 and tex.channels == 2:
+        newBpp = 4
+        newChannels = 4
+
+        logging.debug('\t\tHandling RG16 image...')
+
+        # Pad rows to 32-bit if necessary
+        newStride = tex.width * newBpp
+        remainder = newStride % 4
+        if remainder != 0:
+            newStride = newStride + remainder
+        zeroMip = [0]*(newStride*tex.height)
+
+        for i in range(tex.height):
+            for j in range(tex.width):
+                lumin = tex.pixelData[(tex.height-1-i)*tex.stride + j*tex.bytesPerPixel + 0]
+                alpha = tex.pixelData[(tex.height-1-i)*tex.stride + j*tex.bytesPerPixel + 1]
+
+                zeroMip[i*newStride + j*newBpp + 0] = lumin
+                zeroMip[i*newStride + j*newBpp + 1] = lumin
+                zeroMip[i*newStride + j*newBpp + 2] = lumin
+                zeroMip[i*newStride + j*newBpp + 3] = alpha
+
+        tex.bytesPerPixel = newBpp
+        tex.channels = newChannels
+        tex.stride = newStride
+        tex.pixelData = zeroMip
+        return
+
+    # Pixels are 32-bit aligned (rows might not be)
+    elif validDataType and (4 % tex.bytesPerPixel == 0 or tex.bytesPerPixel % 4 == 0):
+
+        logging.debug('\t\tAligning rows...')
+
+        newStride = tex.stride
+        remainder = tex.stride % 4
+        if remainder != 0:
+            newStride = tex.stride + remainder
+
+        zeroMip = [0]*(newStride*tex.height)
+        for i in range(tex.height):
+            zeroMip[i*newStride:i*newStride + tex.stride] = tex.pixelData[(tex.height-1-i)*tex.stride:(tex.height-i)*tex.stride]
+
+        tex.stride = newStride
+        tex.pixelData = zeroMip
+        return
+
+    # Flip vertically at least, if we have a strange format
+    else:
+        logging.debug('\t\tFlipping vertically...')
+
+        newStride = tex.stride
+        zeroMip = [0]*newStride*tex.height
+        for i in range(tex.height):
+            zeroMip[i*newStride:(i+1)*newStride] = tex.pixelData[(tex.height-1-i)*tex.stride:(tex.height-i)*tex.stride]
+
+        tex.pixelData = zeroMip
+        tex.stride = newStride
+        return
+
+def extractTexture(texName, fieldAccess):
+    """
+    Extracts texture with name 'texName' from fieldAccess and returns its data
+    as a Tex object
+    """
+    vredTex = vrFieldAccess.vrFieldAccess(fieldAccess.getFieldContainer(texName))
+
+    # Ignore disabled/defaulted textures
+    if not vredTex.hasField('useTexture') or \
+        not vredTex.getBool('useTexture') or \
+        not vredTex.hasField('image'):
+        return None
+
+    # Ignore textures without an actual image
+    img = vrFieldAccess.vrFieldAccess(vredTex.getFieldContainer('image'))
+    if not img.isValid():
+        return None
+
+    # Disable all textures so the regular export process doesn't emit them
+    #vredTex.setBool('useTexture', False)
+
+    tex = Tex()
+    tex.texType = texName
+
+    # Get regular properties
+    for texPropName in texPropList:
+        if vredTex.hasField(texPropName):
+            propType = texPropList[texPropName]
+            tex.properties[texPropName] = getPropValue(vredTex, texPropName, propType)
+
+    # Get image properties
+    if img.hasField('name'):
+        tex.srcPath = img.getString('name')
+        tex.width = img.getInt32('width')
+        tex.height = img.getInt32('height')
+        tex.bytesPerPixel = img.getInt32('bpp')
+        tex.dataType = img.getInt32('dataType')
+        tex.pixelFormat = img.getUInt32('pixelFormat')
+        tex.stride = tex.bytesPerPixel * tex.width
+        tex.pixelData = img.getMUInt8('pixel')
+        if tex.pixelFormat == 6406 or tex.pixelFormat == 6409:
+            tex.channels = 1
+        elif tex.pixelFormat == 6407:
+            tex.channels = 3
+        elif tex.pixelFormat == 6408:
+            tex.channels = 4
+        elif tex.pixelFormat == 6410:
+            tex.channels = 2
+
+        unProcessTexture(tex)
+        prepareTexturePixelData(tex)
+
+    # Get projection transform for planar projection
+    proj = vrFieldAccess.vrFieldAccess(vredTex.getFieldContainer('planarProjection'))
+    tex.properties['planarTranslation'] = proj.getVec('translation', 3)
+    tex.properties['planarEulerRotation'] = proj.getVec('eulerRotation', 3)
+    tex.properties['planarScale'] = proj.getVec('scale', 3)
+
+    return tex
+
+def overwriteMappingForTriplanarMaterial(mat, fields):
+    """
+    Triplanar materials are an old, deprecated type of material that is still very
+    common in some scenes. This function will overwrite the texture mapping information
+    for the diffuse, specular, roughness and bump textures present in 'mat' with the
+    data describing a triplanar mapping for a triplanar material
+    """
+    triplanarRepeatU = [1, 1, 1]
+    if fields.hasField('xRepeatU'):
+        triplanarRepeatU[0] = getPropValue(fields, 'xRepeatU', 'SFReal32')
+    if fields.hasField('yRepeatU'):
+        triplanarRepeatU[1] = getPropValue(fields, 'yRepeatU', 'SFReal32')
+    if fields.hasField('zRepeatU'):
+        triplanarRepeatU[2] = getPropValue(fields, 'zRepeatU', 'SFReal32')
+
+    triplanarRepeatV = [1, 1, 1]
+    if fields.hasField('xRepeatV'):
+        triplanarRepeatV[0] = getPropValue(fields, 'xRepeatV', 'SFReal32')
+    if fields.hasField('yRepeatV'):
+        triplanarRepeatV[1] = getPropValue(fields, 'yRepeatV', 'SFReal32')
+    if fields.hasField('zRepeatV'):
+        triplanarRepeatV[2] = getPropValue(fields, 'zRepeatV', 'SFReal32')
+
+    triplanarOffsetU = [0, 0, 0]
+    triplanarOffsetV = [0, 0, 0]
+    if fields.hasField('xOffset'):
+        triplanarOffsetU[0], triplanarOffsetV[0] = getPropValue(fields, 'xOffset', 'SFVec2f')
+    if fields.hasField('yOffset'):
+        triplanarOffsetU[1], triplanarOffsetV[1] = getPropValue(fields, 'yOffset', 'SFVec2f')
+    if fields.hasField('zOffset'):
+        triplanarOffsetU[2], triplanarOffsetV[2] = getPropValue(fields, 'zOffset', 'SFVec2f')
+
+    triplanarRotation = [0, 0, 0]
+    if fields.hasField('rotateX'):
+        triplanarRotation[0] = getPropValue(fields, 'rotateX', 'SFReal32')
+    if fields.hasField('rotateY'):
+        triplanarRotation[1] = getPropValue(fields, 'rotateY', 'SFReal32')
+    if fields.hasField('rotateZ'):
+        triplanarRotation[2] = getPropValue(fields, 'rotateZ', 'SFReal32')
+
+    textureSizeMM = [50, 50]
+    if fields.hasField('textureSizeMM'):
+        textureSizeMM = getPropValue(fields, 'textureSizeMM', 'SFVec2f')
+
+    triplanarBlend = 0
+    if fields.hasField('edgeBlend'):
+        triplanarBlend = getPropValue(fields, 'edgeBlend', 'SFReal32')
+
+    affectedTextures = ['diffuseComponent', 'glossyComponent', 'specularComponent', 'bumpComponent']
+    for texName in affectedTextures:
+        if texName in mat.textures:
+            tex = mat.textures[texName]
+
+            tex.properties['textureProjectionType'] = 2
+            tex.properties['triplanarRepeatU'] = triplanarRepeatU
+            tex.properties['triplanarRepeatV'] = triplanarRepeatV
+            tex.properties['triplanarOffsetU'] = triplanarOffsetU
+            tex.properties['triplanarOffsetV'] = triplanarOffsetV
+            tex.properties['textureSizeMM'] = textureSizeMM
+            tex.properties['triplanarBlend'] = triplanarBlend
+            tex.properties['triplanarRotation'] = triplanarRotation
+
+def extractMaterials(vredMats):
+    """
+    Builds a list of Mat from vrMaterialPtr objects
+
+    Args:
+        vredMats (list of vrMaterialPtr): Materials to process
+
+    Returns:
+        List of Mat objects, one for each material
+    """
+    mats = []
+    processedMats = set()
+
+    for vredMat in vredMats:
+        mat = Mat()
+        mats.append(mat)
+
+        logging.debug(vredMat.getName())
+
+        vredMatId = vredMat.getID()
+        if vredMatId in processedMats:
+            continue
+        processedMats.add(vredMatId)
+
+        mat.name = vredMat.getName()
+        mat.matType = vredMat.getType()
+        mat.vredMatId = vredMatId
+        mat.users = [n.getName() for n in vredMat.getNodes()]
+
+        # Keep track of submaterials if we're a composite material
+        for vredSubmat in vredMat.getSubMaterials():
+            vredSubmatName = vredSubmat.getName()
+            if vredSubmatName not in mat.submaterials:
+                mat.submaterials.append(vredSubmatName)
+
+        # Extract properties outside colorComponentData
+        logging.debug('\tProps outside CCD')
+        fields = vredMat.fields()
+        for fieldPropName in fieldsPropList:
+            if fields.hasField(fieldPropName):
+                propType = fieldsPropList[fieldPropName]
+                mat.properties[fieldPropName] = getPropValue(fields, fieldPropName, propType)
+
+        # Extract data about exterior transparency texture which is an
+        # exception for some reason
+        logging.debug('\tProp about exterior transparency texture')
+        if fields.hasField('exteriorTransparencyComponent'):
+            tex = extractTexture('exteriorTransparencyComponent', fields)
+            if tex != None:
+                mat.textures['exteriorTransparencyComponent'] = tex
+
+        # Extract properties from colorComponentData
+        if fields.hasField('colorComponentData'):
+            ccd = vrFieldAccess.vrFieldAccess(fields.getFieldContainer('colorComponentData'))
+
+            logging.debug('\tCCD props')
+            for ccdPropName in ccdPropList:
+                if ccd.hasField(ccdPropName):
+                    propType = ccdPropList[ccdPropName]
+                    mat.properties[ccdPropName] = getPropValue(ccd, ccdPropName, propType)
+
+            # Extract textures
+            logging.debug('\tTextures')
+            for texName in texTypes:
+                if ccd.hasField(texName):
+                    tex = extractTexture(texName, ccd)
+                    if tex != None:
+                        mat.textures[texName] = tex
+
+        # If this is a triplanar material, fixup the triplanar mapping info for
+        # affected textures
+        logging.debug('\tFixup for triplanar')
+        if vredMat.getType() == 'UTriplanarMaterial':
+            overwriteMappingForTriplanarMaterial(mat, fields)
+
+        # Extract data for metallic carpaint materials
+        # Append a '2' as most of these have the same names as the primary flake layer
+        logging.debug('\tMetallic carpaint second flake layer')
+        if fields.hasField('secondFlakeLayerSettings'):
+            sfl = vrFieldAccess.vrFieldAccess(fields.getFieldContainer('secondFlakeLayerSettings'))
+            for sflPropName in flakeLayerPropList:
+                if sfl.hasField(sflPropName):
+                    propType = flakeLayerPropList[sflPropName]
+                    mat.properties[sflPropName + "2"] = getPropValue(sfl, sflPropName, propType)
+
+    for mat in mats:
+        logging.debug(mat)
+
+    return mats
\ No newline at end of file
diff --git a/Engine/Plugins/Enterprise/DatasmithFBXImporter/Resources/VREDPlugin/DatasmithExporter/sanitizer.py b/Engine/Plugins/Enterprise/DatasmithFBXImporter/Resources/VREDPlugin/DatasmithExporter/sanitizer.py
new file mode 100644
--- /dev/null
+++ b/Engine/Plugins/Enterprise/DatasmithFBXImporter/Resources/VREDPlugin/DatasmithExporter/sanitizer.py
@@ -0,0 +1,103 @@
+import sys
+import os
+import re
+from collections import defaultdict
+
+# Adds the ScriptPlugins folder to path so that we can find the other files
+# in our module
+scriptPluginsFolder = os.path.dirname(os.path.realpath(os.curdir))
+if scriptPluginsFolder not in sys.path:
+    sys.path.append(scriptPluginsFolder)
+
+# Illegal characters and names for nodes, materials, variants and variant sets
+illegalCharacters = re.compile(r"[^A-Za-z0-9_-]+")
+illegalSuffixes = ('_Light', '_Mesh', '_SharedObject')
+illegalNames = ['SceneRoot', 'VRED2Unreal', 'None']
+illegalPrefixes = ('_')
+
+class XMLNodeAdapter:
+    """
+    This class bridges getName/setName/getID functions to
+    xml node get('name') and set('name', val) functions, so that
+    we can use our renameObjectsAndKeepOriginals function with a list
+    of xml nodes
+    """
+    def __init__(self, xmlNode):
+        self.xmlNode = xmlNode
+
+    def getName(self):
+        return self.xmlNode.get('name')
+    def setName(self, value):
+        self.xmlNode.set('name', value)
+    def getID(self):
+        return id(self)
+    def getType(self):
+        return self.xmlNode.tag
+
+    def __str__(self):
+        return self.__repr__()
+    def __repr__(self):
+        return 'Adapter[name="' + self.getName() + '", type="' + self.getType() + '", id="' + str(self.getID()) + '"]'
+
+def sanitize(name):
+    replaced = re.subn(illegalCharacters, '_', name)
+    modified = replaced[1] > 0
+    result = replaced[0]
+
+    if result.endswith(illegalSuffixes) or result in illegalNames:
+        result += "_"
+        modified = True
+
+    if result.startswith(illegalPrefixes):
+        result = 'Object' + result
+        modified = True
+
+    return result, modified
+
+def renameObjectsAndKeepOriginals(objs):
+    """
+    This function iterates over all objects in objs, sanitizing their names
+    and keeping the original names in a dictionary, so that they can be restored later.
+
+    Args:
+        objs (enumerable of T): Enumerable container of hashable objects that have
+                                the 'getName()' and 'setName(value)' functions
+
+    Returns:
+        (dict of T.getID() to str): Dictionary with object IDs and their original names
+    """
+    renamedObjs = {}
+
+    # Replace illegal characters and handle illegal suffixes or illegal names
+    for obj in objs:
+        objID = obj.getID()
+        originalName = obj.getName()
+
+        sanitized = sanitize(originalName)
+        if sanitized[1]:
+            # Just keep track of the original names
+            if objID not in renamedObjs:
+                renamedObjs[objID] = originalName
+
+            obj.setName(sanitized[0])
+
+    # Group objs with the same name in a dict, ignoring case
+    objectNamesToObjs = defaultdict(list)
+    for obj in objs:
+        objectNamesToObjs[obj.getName().lower()].append(obj)
+
+    # Rename objs with the same name
+    for _, objectList in objectNamesToObjs.items():
+        if len(objectList) < 2:
+            continue
+
+        for count, objToRename in enumerate(objectList):
+            objToRenameID = objToRename.getID()
+
+            if objToRenameID not in renamedObjs:
+                renamedObjs[objToRenameID] = objToRename.getName()
+
+            newName = objToRename.getName() + "_" + str(count)
+            objToRename.setName(newName)
+
+    return renamedObjs
\ No newline at end of file
diff --git a/Engine/Plugins/Enterprise/DatasmithFBXImporter/Resources/VREDPlugin/DatasmithExporter/utils.py b/Engine/Plugins/Enterprise/DatasmithFBXImporter/Resources/VREDPlugin/DatasmithExporter/utils.py
new file mode 100644
--- /dev/null
+++ b/Engine/Plugins/Enterprise/DatasmithFBXImporter/Resources/VREDPlugin/DatasmithExporter/utils.py
@@ -0,0 +1,195 @@
+from PySide2 import QtCore, QtGui
+
+import sys
+import os
+import logging
+
+import vrScenegraph
+import vrFileIO
+import vrFileDialog
+import vrNodeUtils
+import vrFieldAccess
+
+# Adds the ScriptPlugins folder to path so that we can find the other files
+# in our module
+scriptPluginsFolder = os.path.dirname(os.path.realpath(os.curdir))
+if scriptPluginsFolder not in sys.path:
+    sys.path.append(scriptPluginsFolder)
+
+def isSharedNode(node):
+    """
+    A node is shared if its straight up a cloned node, or a transformable root, which is
+    a special type of Transform node that can be used to independently move cloned nodes
+    """
+    return vrScenegraph.isCloned(node) or node.hasAttachment('SynchronizationAttachment')
+
+def isNURBS(node):
+    if node is not None:
+        nodeType = node.getType()
+        if nodeType == 'Surface':
+            return True
+        elif nodeType == 'Geometry':
+            nodeFields = node.fields()
+            if nodeFields.hasField('geometryType'):
+                return nodeFields.getUInt32("geometryType") == 2
+
+    return False
+
+def getChildren(node):
+    children = []
+    nChildren = node.getNChildren()
+    for i in range(nChildren):
+        children.append(node.getChild(i))
+    return children
+
+def flatClone(node):
+    children = getChildren(node)
+    vrScenegraph.subChilds(node, children)
+    clone = vrScenegraph.cloneNode(node, True)
+    vrScenegraph.addChilds(node, children)
+    return clone
+
+def resetTransformNode(node):
+    vrNodeUtils.setTransformNodeEulerRotationOrder(node, 'xyz')
+    vrNodeUtils.setTransformNodeRotatePivot(node, 0, 0, 0, False)
+    vrNodeUtils.setTransformNodeRotatePivotTranslation(node, 0, 0, 0)
+    vrNodeUtils.setTransformNodeRotation(node, 0, 0, 0)
+    vrNodeUtils.setTransformNodeRotationOrientation(node, 0, 0, 0)
+    vrNodeUtils.setTransformNodeScale(node, 1, 1, 1)
+    vrNodeUtils.setTransformNodeScalePivot(node, 0, 0, 0, False)
+    vrNodeUtils.setTransformNodeScalePivotTranslation(node, 0, 0, 0)
+    vrNodeUtils.setTransformNodeShear(node, 0, 0, 0)
+    vrNodeUtils.setTransformNodeTranslation(node, 0, 0, 0, False)
+
+def bruteForce(node, filePath):
+    """
+    Uses a dictionary text file containing potential field names (one per line)
+    to guess all the field names that a node has, returning them in a list.
+    """
+    words = []
+    with open(filePath, 'r') as f:
+        words = f.readlines()
+
+    words = [x.strip() for x in words]
+
+    result = []
+    for word in words:
+        if node.hasField(word):
+            result.append(word)
+
+    return result
+
+def isClose(a, b, rel_tol=1e-09, abs_tol=0.0):
+    return abs(a-b) <= max(rel_tol * max(abs(a), abs(b)), abs_tol)
+
+def writeXMLTag(file, indent, name, attributes, values, close=False):
+    """
+    Writes an XML tag line to 'file', including a newline at the end. Examples:
+
+    Input: writeXMLTag(outfile, 2, "TagName", [attr1, attr2], [val1, val2], False)
+    Output:   <TagName attr1="val1" attr2="val2">\n
+
+    Input: writeXMLTag(outfile, 1, "/TagName", [], [], False)
+    Output:  </TagName>\n
+
+    Input: writeXMLTag(outfile, 0, "TagName", [attr1], [[1, 2, 3]], True)
+    Output: <TagName attr1="[1, 2, 3]"/>\n
+
+    Args:
+        file (file): File to write the tag to, as a new line
+        indent (int): Number of spaces before the opening '<'
+        name (str): Name of the XML node
+        attributes (list of obj): List of attribute names
+        values (list of obj): List of corresponding attribute values
+        close (bool): Whether to write a '/' before the closing '>'
+    """
+    file.write(''.join([' ']*indent) + '<' + str(name))
+
+    for pair in zip(attributes, values):
+        file.write(' ' + str(pair[0]) + '="' + str(pair[1]) + '"')
+
+    if close:
+        file.write('/')
+
+    file.write('>\n')
+
+def getNumSceneNURBs():
+    count = 0
+    for node in vrScenegraph.getAllNodes():
+        if isNURBS(node):
+            count += 1
+    return count
+
+def save():
+    """
+    Saves the current scene to a .vpb in the current opened file folder (saving on top of the scene
+    if it can). If the scene hasn't been saved yet, it will prompt a save dialog.
+    """
+    currentScenePath = vrFileIO.getFileIOFilePath()
+
+    if len(currentScenePath) == 0:
+        currentScenePath = vrFileDialog.getSaveFileName("Save As", "", ["vpb(*.vpb)"], True)
+
+    # Force VPB extension
+    currentScenePath = os.path.splitext(currentScenePath)[0] + '.vpb'
+
+    return currentScenePath, vrFileIO.save(currentScenePath)
+
+def getQImageFormat(bytesPerPixel, numChannels):
+    # The current version of PySide2 doesn't support all latest formats like RGBA64 or Grayscale2
+    imgFormat = QtGui.QImage.Format_Invalid
+    if bytesPerPixel == 1 and numChannels == 1:
+        imgFormat = QtGui.QImage.Format_Grayscale8
+    elif bytesPerPixel == 3 and numChannels == 3:
+        imgFormat = QtGui.QImage.Format_RGB888
+    elif bytesPerPixel == 4 and numChannels == 4:
+        imgFormat = QtGui.QImage.Format_RGBA8888
+    return imgFormat
+
+def getUniqueFilename(fullPath):
+    """
+    Appends number suffixes to generate an unique path
+    e.g. From C:/File.txt it might generate C:/File_2.txt
+    """
+    pathNoExt, ext = os.path.splitext(fullPath)
+
+    newFullPath = fullPath
+    appendedIndex = 0
+    while os.path.exists(newFullPath):
+        appendedIndex += 1
+        newFullPath = str(pathNoExt) + '_' + str(appendedIndex) + str(ext)
+
+    return newFullPath
+
+def getAsFieldContainerPtr(node):
+    """
+    Returns a vrFieldAccess object that allows access to properties outside the
+    object core. For example exposing the 'flipped' flag of AnimNodes.
+
+    Might be None
+    """
+    thisID = node.getID()
+
+    # Assumes that Id and container are returned in the same order, which seems to be
+    # the case but I haven't seen any guarantees of
+    # This trick exploits the fact that each node has its core pointing to the node itself
+    # as the parent, at least. When we retrieve the parent, we can get a fieldcontainer though,
+    # allowing access to other properties of a vrNodePtr
+    pairs = zip(node.fields().getMFieldContainerId('parents'), node.fields().getMFieldContainer('parents'))
+    for p in pairs:
+        Id = p[0]
+        fa = vrFieldAccess.vrFieldAccess(p[1])
+
+        if Id == thisID:
+            return fa
+    return None
+
+def getIsNodeFlipped(node):
+    fcp = getAsFieldContainerPtr(node)
+    try:
+        if fcp.getBool('flipped'):
+            return True
+    except BaseException as e:
+        logging.warning(e.message)
+
+    return False
\ No newline at end of file
