Copy up //UE4/Dev-Enterprise/... to //UE4/Dev-Tools-Staging/...

diff --git a/Engine/Plugins/Enterprise/DatasmithCADImporter/Source/CADInterfaces/Private/CADInterfacesModule.cpp b/Engine/Plugins/Enterprise/DatasmithCADImporter/Source/CADInterfaces/Private/CADInterfacesModule.cpp
--- a/Engine/Plugins/Enterprise/DatasmithCADImporter/Source/CADInterfaces/Private/CADInterfacesModule.cpp
+++ b/Engine/Plugins/Enterprise/DatasmithCADImporter/Source/CADInterfaces/Private/CADInterfacesModule.cpp
@@ -1,23 +1,68 @@
 // Copyright 1998-2019 Epic Games, Inc. All Rights Reserved.
 #include "CADInterfacesModule.h"
 
-#include "Modules/ModuleInterface.h"
+#include "HAL/PlatformProcess.h"
+#include "Misc/Paths.h"
 #include "Modules/ModuleManager.h"
 
 #define LOCTEXT_NAMESPACE "CADInterfacesModule"
 
-FCADInterfacesModule& FCADInterfacesModule::Get()
+DEFINE_LOG_CATEGORY(CADInterfaces);
+
+class FCADInterfacesModule : public ICADInterfacesModule
+{
+private:
+	virtual void StartupModule() override;
+	virtual void ShutdownModule() override;
+
+	static void* KernelIOLibHandle;
+};
+
+void* FCADInterfacesModule::KernelIOLibHandle = nullptr;
+
+ICADInterfacesModule& ICADInterfacesModule::Get()
 {
 	return FModuleManager::LoadModuleChecked< FCADInterfacesModule >(CADINTERFACES_MODULE_NAME);
 }
 
-bool FCADInterfacesModule::IsAvailable()
+bool ICADInterfacesModule::IsAvailable()
 {
 	return FModuleManager::Get().IsModuleLoaded(CADINTERFACES_MODULE_NAME);
 }
 
 void FCADInterfacesModule::StartupModule()
 {
+	check(KernelIOLibHandle == nullptr);
+
+	FString KernelIODll = TEXT("kernel_io.dll");
+
+	// determine directory paths
+	FString CADImporterDllPath = FPaths::Combine(FPaths::EnginePluginsDir(), TEXT("Enterprise/DatasmithCADImporter"), TEXT("Binaries"), FPlatformProcess::GetBinariesSubdirectory());
+	FPlatformProcess::PushDllDirectory(*CADImporterDllPath);
+	KernelIODll = FPaths::Combine(CADImporterDllPath, KernelIODll);
+
+	if (!FPaths::FileExists(KernelIODll))
+	{
+		UE_LOG(CADInterfaces, Error, TEXT("Failed to find the binary folder for the CoreTech dll. Plug-in will not be functional."));
+		return;
+	}
+
+	KernelIOLibHandle = FPlatformProcess::GetDllHandle(*KernelIODll);
+	if (KernelIOLibHandle == nullptr)
+	{
+		UE_LOG(CADInterfaces, Error, TEXT("Failed to load required library %s. Plug-in will not be functional."), *KernelIODll);
+	}
+
+	FPlatformProcess::PopDllDirectory(*CADImporterDllPath);
+}
+
+void FCADInterfacesModule::ShutdownModule()
+{
+	if (KernelIOLibHandle != nullptr)
+	{
+		FPlatformProcess::FreeDllHandle(KernelIOLibHandle);
+		KernelIOLibHandle = nullptr;
+	}
 }
 
 IMPLEMENT_MODULE(FCADInterfacesModule, CADInterfaces);
diff --git a/Engine/Plugins/Enterprise/DatasmithCADImporter/Source/CADInterfaces/Public/CADInterfacesModule.h b/Engine/Plugins/Enterprise/DatasmithCADImporter/Source/CADInterfaces/Public/CADInterfacesModule.h
--- a/Engine/Plugins/Enterprise/DatasmithCADImporter/Source/CADInterfaces/Public/CADInterfacesModule.h
+++ b/Engine/Plugins/Enterprise/DatasmithCADImporter/Source/CADInterfaces/Public/CADInterfacesModule.h
@@ -1,16 +1,16 @@
 // Copyright 1998-2019 Epic Games, Inc. All Rights Reserved.
 #pragma once
 
+#include "CoreMinimal.h"
 #include "Modules/ModuleInterface.h"
 
-#define CADINTERFACES_MODULE_NAME TEXT("WorkerDataManagement")
+#define CADINTERFACES_MODULE_NAME TEXT("CADInterfaces")
 
-class FCADInterfacesModule : public IModuleInterface
+DECLARE_LOG_CATEGORY_EXTERN(CADInterfaces, Log, All);
+
+class ICADInterfacesModule : public IModuleInterface
 {
 public:
-	static FCADInterfacesModule& Get();
+	static ICADInterfacesModule& Get();
 	static bool IsAvailable();
-
-private:
-	virtual void StartupModule() override;
 };
diff --git a/Engine/Plugins/Enterprise/DatasmithCADImporter/Source/CADLibrary/Private/CoreTechHelper.cpp b/Engine/Plugins/Enterprise/DatasmithCADImporter/Source/CADLibrary/Private/CoreTechHelper.cpp
--- a/Engine/Plugins/Enterprise/DatasmithCADImporter/Source/CADLibrary/Private/CoreTechHelper.cpp
+++ b/Engine/Plugins/Enterprise/DatasmithCADImporter/Source/CADLibrary/Private/CoreTechHelper.cpp
@@ -657,36 +657,6 @@ TSharedPtr<IDatasmithUEPbrMaterialElement> CreateUEPbrMaterialFromMaterial(FCADM
 	return MaterialElement;
 }
 
-uint32 GetSize(CT_TESS_DATA_TYPE type)
-{
-	switch (type)
-	{
-	case CT_TESS_USE_DEFAULT:
-		return sizeof(uint32);
-	case CT_TESS_UBYTE:
-		return sizeof(uint8_t);
-	case CT_TESS_BYTE:
-		return sizeof(int8_t);
-	case CT_TESS_USHORT:
-		return sizeof(int16_t);
-	case CT_TESS_SHORT:
-		return sizeof(uint16_t);
-	case CT_TESS_UINT:
-		return sizeof(uint32);
-	case CT_TESS_INT:
-		return sizeof(int32);
-	case CT_TESS_ULONG:
-		return sizeof(uint64);
-	case CT_TESS_LONG:
-		return sizeof(int64);
-	case CT_TESS_FLOAT:
-		return sizeof(float);
-	case CT_TESS_DOUBLE:
-		return sizeof(double);
-	}
-	return 0;
-}
-
 CT_IO_ERROR Tessellate(CT_OBJECT_ID MainObjectId, const FImportParameters& ImportParams, FMeshDescription& MeshDesc, FMeshParameters& MeshParameters)
 {
 	CheckedCTError Result;
diff --git a/Engine/Plugins/Enterprise/DatasmithCADImporter/Source/CADLibrary/Public/CoreTechHelper.h b/Engine/Plugins/Enterprise/DatasmithCADImporter/Source/CADLibrary/Public/CoreTechHelper.h
--- a/Engine/Plugins/Enterprise/DatasmithCADImporter/Source/CADLibrary/Public/CoreTechHelper.h
+++ b/Engine/Plugins/Enterprise/DatasmithCADImporter/Source/CADLibrary/Public/CoreTechHelper.h
@@ -48,8 +48,6 @@ namespace CADLibrary
 		}
 	};
 
-	uint32 GetSize(CT_TESS_DATA_TYPE type);
-
 	CADLIBRARY_API bool ConvertCTBodySetToMeshDescription(const FImportParameters& ImportParams, const FMeshParameters& MeshParameters, FBodyMesh& Body, FMeshDescription& MeshDescription);
 
 	CADLIBRARY_API CT_IO_ERROR Tessellate(CT_OBJECT_ID MainObjectId, const FImportParameters& ImportParams, FMeshDescription& Mesh, FMeshParameters& MeshParameters);
diff --git a/Engine/Plugins/Enterprise/DatasmithCADImporter/Source/DatasmithCADTranslator/DatasmithCADTranslator.Build.cs b/Engine/Plugins/Enterprise/DatasmithCADImporter/Source/DatasmithCADTranslator/DatasmithCADTranslator.Build.cs
--- a/Engine/Plugins/Enterprise/DatasmithCADImporter/Source/DatasmithCADTranslator/DatasmithCADTranslator.Build.cs
+++ b/Engine/Plugins/Enterprise/DatasmithCADImporter/Source/DatasmithCADTranslator/DatasmithCADTranslator.Build.cs
@@ -23,7 +23,7 @@ namespace UnrealBuildTool.Rules
 					"CADLibrary",
 					"CADTools",
 					"DatasmithCore",
-					"DatasmithCoreTechExtension",
+					"DatasmithCoreTechParametricSurfaceData",
 					"DatasmithContent",
 					"DatasmithImporter",
 					"DatasmithDispatcher",
diff --git a/Engine/Plugins/Enterprise/DatasmithCADImporter/Source/DatasmithCADTranslator/Private/DatasmithSceneGraphBuilder.cpp b/Engine/Plugins/Enterprise/DatasmithCADImporter/Source/DatasmithCADTranslator/Private/DatasmithSceneGraphBuilder.cpp
--- a/Engine/Plugins/Enterprise/DatasmithCADImporter/Source/DatasmithCADTranslator/Private/DatasmithSceneGraphBuilder.cpp
+++ b/Engine/Plugins/Enterprise/DatasmithCADImporter/Source/DatasmithCADTranslator/Private/DatasmithSceneGraphBuilder.cpp
@@ -208,6 +208,30 @@ TSharedPtr< IDatasmithActorElement >  FDatasmithSceneGraphBuilder::BuildInstance
 	return Actor;
 }
 
+void FDatasmithSceneGraphBuilder::FillAnchorActor(const TSharedRef<IDatasmithActorElement>& ActorElement, const FString& CleanFilenameOfCADFile)
+{
+	CurrentMockUp = CADFileToArchiveMockUp.FindRef(CleanFilenameOfCADFile);
+	CadId RootId = 1;
+	int32* Index = CurrentMockUp->CADIdToComponentIndex.Find(RootId);
+	ActorData Data(TEXT(""));
+
+	// TODO: check ParentData and Index validity?
+	ActorData ParentData(ActorElement->GetName());
+	CADLibrary::FArchiveComponent& Component = CurrentMockUp->ComponentSet[*Index];
+
+	TMap<FString, FString> InstanceNodeMetaDataMap;
+	FString ActorUUID;
+	FString ActorLabel;
+	GetNodeUUIDAndName(InstanceNodeMetaDataMap, Component.MetaData, ParentData.Uuid, ActorUUID, ActorLabel);
+
+	AddMetaData(ActorElement, InstanceNodeMetaDataMap, Component.MetaData);
+
+	ActorData ComponentData(*ActorUUID, ParentData);
+	GetMainMaterial(Component.MetaData, ComponentData, bMaterialPropagationIsTopDown);
+
+	AddChildren(ActorElement, Component, ComponentData);
+}
+
 TSharedPtr< IDatasmithActorElement >  FDatasmithSceneGraphBuilder::CreateActor(const TCHAR* InEUUID, const TCHAR* InLabel)
 {
 	TSharedPtr< IDatasmithActorElement > Actor = FDatasmithSceneFactory::CreateActor(InEUUID);
diff --git a/Engine/Plugins/Enterprise/DatasmithCADImporter/Source/DatasmithCADTranslator/Public/DatasmithMeshBuilder.h b/Engine/Plugins/Enterprise/DatasmithCADImporter/Source/DatasmithCADTranslator/Public/DatasmithMeshBuilder.h
new file mode 100644
--- /dev/null
+++ b/Engine/Plugins/Enterprise/DatasmithCADImporter/Source/DatasmithCADTranslator/Public/DatasmithMeshBuilder.h
@@ -0,0 +1,33 @@
+// Copyright 1998-2019 Epic Games, Inc. All Rights Reserved.
+#pragma once
+
+#include "CoreMinimal.h"
+
+#include "CADData.h"
+#include "CADOptions.h"
+
+#include "MeshDescription.h"
+
+
+class IDatasmithMeshElement;
+
+
+class DATASMITHCADTRANSLATOR_API FDatasmithMeshBuilder
+{
+public:
+	FDatasmithMeshBuilder(TMap<FString, FString>& InCADFileToMeshFileMap, const FString& InCachePath, const CADLibrary::FImportParameters& InImportParameters);
+
+	TOptional<FMeshDescription> GetMeshDescription(TSharedRef<IDatasmithMeshElement> OutMeshElement, CADLibrary::FMeshParameters& OutMeshParameters);
+
+protected:
+	FString CachePath;
+
+	void LoadMeshFiles();
+
+	TMap<FString, FString>& CADFileToMeshFile;
+	TArray<TArray<CADLibrary::FBodyMesh>> BodyMeshes;
+	TMap<CADUUID, CADLibrary::FBodyMesh*> MeshActorNameToBodyMesh;
+
+	CADLibrary::FImportParameters ImportParameters;
+};
+
diff --git a/Engine/Plugins/Enterprise/DatasmithCADImporter/Source/DatasmithCADTranslator/Public/DatasmithSceneGraphBuilder.h b/Engine/Plugins/Enterprise/DatasmithCADImporter/Source/DatasmithCADTranslator/Public/DatasmithSceneGraphBuilder.h
new file mode 100644
--- /dev/null
+++ b/Engine/Plugins/Enterprise/DatasmithCADImporter/Source/DatasmithCADTranslator/Public/DatasmithSceneGraphBuilder.h
@@ -0,0 +1,137 @@
+// Copyright 1998-2019 Epic Games, Inc. All Rights Reserved.
+#pragma once
+
+#ifdef CAD_INTERFACE
+#include "CoreMinimal.h"
+
+#include "CADData.h"
+#include "CADSceneGraph.h"
+#include "CoreTechHelper.h"
+#include "CoreTechTypes.h"
+#include "CTSession.h"
+#include "DatasmithImportOptions.h"
+
+#include "Containers/Map.h"
+#include "Containers/Queue.h"
+#include "Misc/Paths.h"
+
+static const FString SPACE(" ");
+
+class FCADSceneGraphDescriptionFile;
+class FDatasmithSceneSource;
+class IDatasmithActorElement;
+class IDatasmithMeshElement;
+class IDatasmithScene;
+class IDatasmithUEPbrMaterialElement;
+
+enum NODE_TYPE
+{
+	INSTANCE,
+	COMPONENT,
+	EXTERNALCOMPONENT,
+	BODY,
+	UNDEFINED,
+	UNKNOWN
+};
+
+namespace CADLibrary		
+{
+class FArchiveMockUp;
+class ICADArchiveObject;
+class FArchiveInstance;
+class FArchiveComponent;
+class FArchiveBody;
+class FArchiveMaterial;
+class FArchiveColor;
+}
+
+class ActorData  //#ueent_CAD
+{
+public:
+	ActorData(const TCHAR* NodeUuid, const ActorData& ParentData)
+		: Uuid(NodeUuid)
+		, Material(ParentData.Material)
+		, MaterialUuid(ParentData.MaterialUuid)
+		, Color(ParentData.Color)
+		, ColorUuid(ParentData.ColorUuid)
+	{
+	}
+
+	ActorData(const TCHAR* NodeUuid)
+		: Uuid(NodeUuid)
+		, MaterialUuid(0)
+		, ColorUuid(0)
+	{
+	}
+
+	const TCHAR* Uuid;
+
+	CADLibrary::FCADMaterial Material;
+	uint32 MaterialUuid;
+
+	FColor Color;
+	uint32 ColorUuid;
+};
+
+
+
+
+class DATASMITHCADTRANSLATOR_API FDatasmithSceneGraphBuilder
+{
+public:
+	FDatasmithSceneGraphBuilder(
+		TMap<FString, FString>& InCADFileToUE4FileMap, 
+		const FString& InCachePath, 
+		TSharedRef<IDatasmithScene> InScene, 
+		const FDatasmithSceneSource& InSource, 
+		const CADLibrary::FImportParameters& InImportParameters);
+
+	bool Build();
+
+	void LoadSceneGraphDescriptionFiles();
+
+	void FillAnchorActor(const TSharedRef< IDatasmithActorElement >& ActorElement, const FString& CleanFilenameOfCADFile);
+private:
+
+	TSharedPtr< IDatasmithActorElement > BuildInstance(CADLibrary::FArchiveInstance& Instance, ActorData& ParentData);
+	TSharedPtr< IDatasmithActorElement > BuildComponent(CADLibrary::FArchiveComponent& Component, ActorData& ParentData);
+	TSharedPtr< IDatasmithActorElement > BuildBody(CADLibrary::FArchiveBody& Body, ActorData& ParentData);
+
+	void AddMetaData(TSharedPtr< IDatasmithActorElement > ActorElement, TMap<FString, FString>& InstanceNodeAttributeSetMap, TMap<FString, FString>& ReferenceNodeAttributeSetMap);
+	void AddChildren(TSharedPtr< IDatasmithActorElement > Actor, CADLibrary::FArchiveComponent& Component, ActorData& ParentData);
+	bool DoesActorHaveChildrenOrIsAStaticMesh(const TSharedPtr< IDatasmithActorElement >& ActorElement);
+
+	TSharedPtr< IDatasmithUEPbrMaterialElement > GetDefaultMaterial();
+	TSharedPtr<IDatasmithMaterialIDElement> FindOrAddMaterial(uint32 MaterialUuid);
+
+	TSharedPtr< IDatasmithActorElement > CreateActor(const TCHAR* ActorUUID, const TCHAR* ActorLabel);
+	TSharedPtr< IDatasmithMeshElement > FindOrAddMeshElement(CADLibrary::FArchiveBody& Body, FString& InLabel);
+
+	void GetNodeUUIDAndName(TMap<FString, FString>& InInstanceNodeMetaDataMap, TMap<FString, FString>& InReferenceNodeMetaDataMap, const TCHAR* InParentUEUUID, FString& OutUEUUID, FString& OutName);
+
+protected:
+	TMap<FString, FString>& CADFileToSceneGraphDescriptionFile;
+	const FString& CachePath;
+	TSharedRef<IDatasmithScene> DatasmithScene;
+	const FDatasmithSceneSource& Source;
+	const CADLibrary::FImportParameters& ImportParameters;
+	const uint32 ImportParametersHash;
+
+	TArray<CADLibrary::FArchiveMockUp> ArchiveMockUps;
+	TMap<FString, CADLibrary::FArchiveMockUp*> CADFileToArchiveMockUp;
+
+	TMap< CADUUID, TSharedPtr< IDatasmithMeshElement > > BodyUuidToMeshElement;
+
+	TMap< CADUUID, TSharedPtr< IDatasmithUEPbrMaterialElement > > MaterialUuidMap;
+	TSharedPtr<IDatasmithUEPbrMaterialElement > DefaultMaterial;
+
+	TMap<CADUUID, CADLibrary::FArchiveColor> ColorNameToColorArchive; 
+	TMap<CADUUID, CADLibrary::FArchiveMaterial> MaterialNameToMaterialArchive; 
+
+	CADLibrary::FArchiveMockUp* CurrentMockUp;
+
+	bool bPreferMaterial;
+	bool bMaterialPropagationIsTopDown;
+};
+
+#endif // CAD_INTERFACE
diff --git a/Engine/Plugins/Enterprise/DatasmithCADImporter/Source/DatasmithCoreTechExtension/DatasmithCoreTechExtension.Build.cs b/Engine/Plugins/Enterprise/DatasmithCADImporter/Source/DatasmithCoreTechExtension/DatasmithCoreTechExtension.Build.cs
--- a/Engine/Plugins/Enterprise/DatasmithCADImporter/Source/DatasmithCoreTechExtension/DatasmithCoreTechExtension.Build.cs
+++ b/Engine/Plugins/Enterprise/DatasmithCADImporter/Source/DatasmithCoreTechExtension/DatasmithCoreTechExtension.Build.cs
@@ -23,6 +23,7 @@ namespace UnrealBuildTool.Rules
 					"MeshDescription",
 					"DataprepCore",
 					"DatasmithContent",
+					"DatasmithCoreTechParametricSurfaceData",
 					"DatasmithImporter",
 					"Engine",
 					"StaticMeshEditor",
diff --git a/Engine/Plugins/Enterprise/DatasmithCADImporter/Source/DatasmithCoreTechExtension/Private/CoreTechBlueprintLibrary.cpp b/Engine/Plugins/Enterprise/DatasmithCADImporter/Source/DatasmithCoreTechExtension/Private/CoreTechBlueprintLibrary.cpp
--- a/Engine/Plugins/Enterprise/DatasmithCADImporter/Source/DatasmithCoreTechExtension/Private/CoreTechBlueprintLibrary.cpp
+++ b/Engine/Plugins/Enterprise/DatasmithCADImporter/Source/DatasmithCoreTechExtension/Private/CoreTechBlueprintLibrary.cpp
@@ -6,7 +6,7 @@
 #include "CoreTechHelper.h"
 #include "CoreTechMeshLoader.h"
 #endif
-#include "CoreTechParametricSurfaceExtension.h"
+#include "CoreTechRetessellateAction.h"
 #include "DatasmithAdditionalData.h"
 #include "DatasmithStaticMeshImporter.h" // Call to BuildStaticMesh
 #include "DatasmithUtils.h"
diff --git a/Engine/Plugins/Enterprise/DatasmithCADImporter/Source/DatasmithCoreTechExtension/Private/DatasmithCoreTechExtensionModule.cpp b/Engine/Plugins/Enterprise/DatasmithCADImporter/Source/DatasmithCoreTechExtension/Private/DatasmithCoreTechExtensionModule.cpp
--- a/Engine/Plugins/Enterprise/DatasmithCADImporter/Source/DatasmithCoreTechExtension/Private/DatasmithCoreTechExtensionModule.cpp
+++ b/Engine/Plugins/Enterprise/DatasmithCADImporter/Source/DatasmithCoreTechExtension/Private/DatasmithCoreTechExtensionModule.cpp
@@ -2,7 +2,7 @@
 
 #include "DatasmithCoreTechExtensionModule.h"
 
-#include "CoreTechParametricSurfaceExtension.h"
+#include "CoreTechRetessellateAction.h"
 
 #include "Framework/MultiBox/MultiBoxBuilder.h"
 #include "IStaticMeshEditor.h"
diff --git a/Engine/Plugins/Enterprise/DatasmithCADImporter/Source/DatasmithCoreTechExtension/Public/CoreTechRetessellateAction.h b/Engine/Plugins/Enterprise/DatasmithCADImporter/Source/DatasmithCoreTechExtension/Public/CoreTechRetessellateAction.h
new file mode 100644
--- /dev/null
+++ b/Engine/Plugins/Enterprise/DatasmithCADImporter/Source/DatasmithCoreTechExtension/Public/CoreTechRetessellateAction.h
@@ -0,0 +1,55 @@
+// Copyright 1998-2019 Epic Games, Inc. All Rights Reserved.
+
+#pragma once
+
+#include "CoreMinimal.h"
+
+#include "CoreTechParametricSurfaceExtension.h"
+#include "DatasmithCustomAction.h"
+#include "DatasmithAdditionalData.h"
+#include "DatasmithImportOptions.h"
+
+#include "CoreTechRetessellateAction.generated.h"
+
+class DATASMITHCORETECHEXTENSION_API FCoreTechRetessellate_Impl
+{
+public:
+	static const FText Label;
+	static const FText Tooltip;
+
+	static bool CanApplyOnAssets(const TArray<FAssetData>& SelectedAssets);
+	static void ApplyOnAssets(const TArray<FAssetData>& SelectedAssets);
+	static bool ApplyOnOneAsset(UStaticMesh& StaticMesh, UCoreTechParametricSurfaceData& CoreTechData, const FDatasmithTessellationOptions& RetesselateOptions);
+};
+
+
+UCLASS()
+class DATASMITHCORETECHEXTENSION_API UCoreTechRetessellateAction : public UDatasmithCustomActionBase
+{
+	GENERATED_BODY()
+
+public:
+	virtual const FText& GetLabel() override { return FCoreTechRetessellate_Impl::Label; }
+	virtual const FText& GetTooltip() override { return FCoreTechRetessellate_Impl::Tooltip; }
+
+	virtual bool CanApplyOnAssets(const TArray<FAssetData>& SelectedAssets) override
+	{ return FCoreTechRetessellate_Impl::CanApplyOnAssets(SelectedAssets); }
+
+	virtual void ApplyOnAssets(const TArray<FAssetData>& SelectedAssets) override
+	{ return FCoreTechRetessellate_Impl::ApplyOnAssets(SelectedAssets); }
+
+	virtual bool CanApplyOnActors(const TArray<AActor*>& SelectedActors) override;
+
+	virtual void ApplyOnActors(const TArray<AActor*>& SelectedActors) override;
+};
+
+
+UCLASS(config = Editor, Transient)
+class DATASMITHCORETECHEXTENSION_API UCoreTechRetessellateActionOptions : public UObject
+{
+	GENERATED_BODY()
+
+public:
+	UPROPERTY(config, EditAnywhere, BlueprintReadWrite, Category = "NotVisible", meta = (ShowOnlyInnerProperties))
+	FDatasmithTessellationOptions Options;
+};
diff --git a/Engine/Plugins/Enterprise/DatasmithCADImporter/Source/DatasmithCoreTechParametricSurfaceData/DatasmithCoreTechParametricSurfaceData.Build.cs b/Engine/Plugins/Enterprise/DatasmithCADImporter/Source/DatasmithCoreTechParametricSurfaceData/DatasmithCoreTechParametricSurfaceData.Build.cs
new file mode 100644
--- /dev/null
+++ b/Engine/Plugins/Enterprise/DatasmithCADImporter/Source/DatasmithCoreTechParametricSurfaceData/DatasmithCoreTechParametricSurfaceData.Build.cs
@@ -0,0 +1,33 @@
+// Copyright 1998-2019 Epic Games, Inc. All Rights Reserved.
+
+using System.IO;
+
+namespace UnrealBuildTool.Rules
+{
+	public class DatasmithCoreTechParametricSurfaceData : ModuleRules
+	{
+		public DatasmithCoreTechParametricSurfaceData(ReadOnlyTargetRules Target) : base(Target)
+		{
+			PublicDependencyModuleNames.AddRange(
+				new string[]
+				{
+					"Core",
+				}
+			);
+
+			PrivateDependencyModuleNames.AddRange(
+				new string[]
+				{
+					"Core",
+					"CoreUObject",
+					"MeshDescription",
+					"DatasmithContent",
+					"DatasmithImporter",
+					"Engine",
+					"StaticMeshDescription",
+					"CADLibrary",
+				}
+			);
+		}
+	}
+}
diff --git a/Engine/Plugins/Enterprise/DatasmithCADImporter/Source/DatasmithCoreTechParametricSurfaceData/Private/DatasmithCoreTechParametricSurfaceDataModule.cpp b/Engine/Plugins/Enterprise/DatasmithCADImporter/Source/DatasmithCoreTechParametricSurfaceData/Private/DatasmithCoreTechParametricSurfaceDataModule.cpp
new file mode 100644
--- /dev/null
+++ b/Engine/Plugins/Enterprise/DatasmithCADImporter/Source/DatasmithCoreTechParametricSurfaceData/Private/DatasmithCoreTechParametricSurfaceDataModule.cpp
@@ -0,0 +1,24 @@
+// Copyright 1998-2019 Epic Games, Inc. All Rights Reserved.
+
+#include "DatasmithCoreTechParametricSurfaceDataModule.h"
+
+#include "Modules/ModuleInterface.h"
+#include "Modules/ModuleManager.h"
+
+
+#define LOCTEXT_NAMESPACE "DatasmithCoreTechParametricSurfaceDataModule"
+
+FDatasmithCoreTechParametricSurfaceDataModule& FDatasmithCoreTechParametricSurfaceDataModule::Get()
+{
+	return FModuleManager::LoadModuleChecked< FDatasmithCoreTechParametricSurfaceDataModule >(DATASMITHCORETECHPARAMETRICSURFACEDATA_MODULE_NAME);
+}
+
+bool FDatasmithCoreTechParametricSurfaceDataModule::IsAvailable()
+{
+	return FModuleManager::Get().IsModuleLoaded(DATASMITHCORETECHPARAMETRICSURFACEDATA_MODULE_NAME);
+}
+
+IMPLEMENT_MODULE(FDatasmithCoreTechParametricSurfaceDataModule, DatasmithCoreTechParametricSurfaceData);
+
+#undef LOCTEXT_NAMESPACE // "DatasmithCoreTechParametricSurfaceDataModule"
+
diff --git a/Engine/Plugins/Enterprise/DatasmithCADImporter/Source/DatasmithCoreTechParametricSurfaceData/Public/CoreTechParametricSurfaceExtension.h b/Engine/Plugins/Enterprise/DatasmithCADImporter/Source/DatasmithCoreTechParametricSurfaceData/Public/CoreTechParametricSurfaceExtension.h
new file mode 100644
--- /dev/null
+++ b/Engine/Plugins/Enterprise/DatasmithCADImporter/Source/DatasmithCoreTechParametricSurfaceData/Public/CoreTechParametricSurfaceExtension.h
@@ -0,0 +1,69 @@
+// Copyright 1998-2019 Epic Games, Inc. All Rights Reserved.
+
+#pragma once
+
+#include "CoreMinimal.h"
+
+#include "DatasmithCustomAction.h"
+#include "DatasmithAdditionalData.h"
+#include "DatasmithImportOptions.h"
+
+#include "CoreTechParametricSurfaceExtension.generated.h"
+
+
+USTRUCT(BlueprintType)
+struct DATASMITHCORETECHPARAMETRICSURFACEDATA_API FCoreTechSceneParameters
+{
+	GENERATED_BODY()
+
+	// value from FDatasmithUtils::EModelCoordSystem
+	UPROPERTY()
+	uint8 ModelCoordSys;
+
+	UPROPERTY()
+	float MetricUnit;
+
+	UPROPERTY()
+	float ScaleFactor;
+};
+
+USTRUCT()
+struct DATASMITHCORETECHPARAMETRICSURFACEDATA_API FCoreTechMeshParameters
+{
+	GENERATED_BODY()
+
+	UPROPERTY()
+	bool bNeedSwapOrientation;
+
+	UPROPERTY()
+	bool bIsSymmetric;
+
+	UPROPERTY()
+	FVector SymmetricOrigin;
+
+	UPROPERTY()
+	FVector SymmetricNormal;
+};
+
+
+UCLASS()
+class DATASMITHCORETECHPARAMETRICSURFACEDATA_API UCoreTechParametricSurfaceData : public UDatasmithAdditionalData
+{
+	GENERATED_BODY()
+
+public:
+	UPROPERTY()
+	FString SourceFile;
+
+	UPROPERTY()
+	TArray<uint8> RawData;
+
+	UPROPERTY()
+	FCoreTechSceneParameters SceneParameters;
+
+	UPROPERTY()
+	FCoreTechMeshParameters MeshParameters;
+
+	UPROPERTY(EditAnywhere, Category=NURBS)
+	FDatasmithTessellationOptions LastTessellationOptions;
+};
\ No newline at end of file
diff --git a/Engine/Plugins/Enterprise/DatasmithCADImporter/Source/DatasmithCoreTechParametricSurfaceData/Public/DatasmithCoreTechParametricSurfaceDataModule.h b/Engine/Plugins/Enterprise/DatasmithCADImporter/Source/DatasmithCoreTechParametricSurfaceData/Public/DatasmithCoreTechParametricSurfaceDataModule.h
new file mode 100644
--- /dev/null
+++ b/Engine/Plugins/Enterprise/DatasmithCADImporter/Source/DatasmithCoreTechParametricSurfaceData/Public/DatasmithCoreTechParametricSurfaceDataModule.h
@@ -0,0 +1,17 @@
+// Copyright 1998-2019 Epic Games, Inc. All Rights Reserved.
+
+#pragma once
+
+#include "Modules/ModuleInterface.h"
+
+#define DATASMITHCORETECHPARAMETRICSURFACEDATA_MODULE_NAME TEXT("DatasmithCoreTechParametricSurfaceData")
+
+/**
+ * This module exposes additional features for assets containing CoreTech data.
+ */
+class FDatasmithCoreTechParametricSurfaceDataModule : public IModuleInterface
+{
+public:
+    static FDatasmithCoreTechParametricSurfaceDataModule& Get();
+    static bool IsAvailable();
+};
diff --git a/Engine/Plugins/Enterprise/DatasmithCADImporter/Source/DatasmithOpenNurbsTranslator/DatasmithOpenNurbsTranslator.Build.cs b/Engine/Plugins/Enterprise/DatasmithCADImporter/Source/DatasmithOpenNurbsTranslator/DatasmithOpenNurbsTranslator.Build.cs
--- a/Engine/Plugins/Enterprise/DatasmithCADImporter/Source/DatasmithOpenNurbsTranslator/DatasmithOpenNurbsTranslator.Build.cs
+++ b/Engine/Plugins/Enterprise/DatasmithCADImporter/Source/DatasmithOpenNurbsTranslator/DatasmithOpenNurbsTranslator.Build.cs
@@ -24,7 +24,7 @@ namespace UnrealBuildTool.Rules
                     "CADLibrary",
 					"Engine",
 					"DatasmithCore",
-                    "DatasmithCoreTechExtension",
+                    "DatasmithCoreTechParametricSurfaceData",
 					"DatasmithContent",
 					"DatasmithImporter",
 					"MeshDescription",
diff --git a/Engine/Plugins/Enterprise/DatasmithContent/Source/DatasmithContent/Private/ObjectTemplates/DatasmithStaticMeshTemplate.cpp b/Engine/Plugins/Enterprise/DatasmithContent/Source/DatasmithContent/Private/ObjectTemplates/DatasmithStaticMeshTemplate.cpp
--- a/Engine/Plugins/Enterprise/DatasmithContent/Source/DatasmithContent/Private/ObjectTemplates/DatasmithStaticMeshTemplate.cpp
+++ b/Engine/Plugins/Enterprise/DatasmithContent/Source/DatasmithContent/Private/ObjectTemplates/DatasmithStaticMeshTemplate.cpp
@@ -20,9 +20,9 @@ void FDatasmithMeshBuildSettingsTemplate::Apply( FMeshBuildSettings* Destination
 
 	// The settings for RecomputeNormals and RecomputeTangents when True must be honored irrespective of the previous template settings
 	// because their values are determined by ShouldRecomputeNormals/ShouldRecomputeTangents which determine if they are needed by the renderer
-	Destination->bRecomputeNormals = PreviousTemplate ? Destination->bRecomputeNormals | bRecomputeNormals : bRecomputeNormals;
+	Destination->bRecomputeNormals = PreviousTemplate ? PreviousTemplate->bRecomputeNormals | bRecomputeNormals : bRecomputeNormals;
 
-	Destination->bRecomputeTangents = PreviousTemplate ? Destination->bRecomputeTangents | bRecomputeTangents : bRecomputeTangents;
+	Destination->bRecomputeTangents = PreviousTemplate ? PreviousTemplate->bRecomputeTangents | bRecomputeTangents : bRecomputeTangents;
 
 	DATASMITHOBJECTTEMPLATE_CONDITIONALSET( bRemoveDegenerates, Destination, PreviousTemplate );
 
diff --git a/Engine/Plugins/Enterprise/DatasmithFBXImporter/Source/DatasmithDeltaGenTranslator/Private/DatasmithDeltaGenImportData.h b/Engine/Plugins/Enterprise/DatasmithFBXImporter/Source/DatasmithDeltaGenTranslator/Private/DatasmithDeltaGenImportData.h
--- a/Engine/Plugins/Enterprise/DatasmithFBXImporter/Source/DatasmithDeltaGenTranslator/Private/DatasmithDeltaGenImportData.h
+++ b/Engine/Plugins/Enterprise/DatasmithFBXImporter/Source/DatasmithDeltaGenTranslator/Private/DatasmithDeltaGenImportData.h
@@ -159,6 +159,9 @@ struct FDeltaGenTmlDataTimelineAnimation
 struct FDeltaGenTmlDataTimeline : public FTableRowBase
 {
 	FString Name;
+	float Framerate;
+
+	UPROPERTY(BlueprintReadWrite, EditAnywhere, Category=DeltaGen)
 	TArray<FDeltaGenTmlDataTimelineAnimation> Animations;
 };
 
diff --git a/Engine/Plugins/Enterprise/DatasmithFBXImporter/Source/DatasmithDeltaGenTranslator/Private/DatasmithDeltaGenImporter.cpp b/Engine/Plugins/Enterprise/DatasmithFBXImporter/Source/DatasmithDeltaGenTranslator/Private/DatasmithDeltaGenImporter.cpp
--- a/Engine/Plugins/Enterprise/DatasmithFBXImporter/Source/DatasmithDeltaGenTranslator/Private/DatasmithDeltaGenImporter.cpp
+++ b/Engine/Plugins/Enterprise/DatasmithFBXImporter/Source/DatasmithDeltaGenTranslator/Private/DatasmithDeltaGenImporter.cpp
@@ -731,7 +731,7 @@ TSharedPtr<IDatasmithBaseMaterialElement> FDatasmithDeltaGenImporter::ConvertMat
 	return MaterialElement;
 }
 
-void PopulateTransformAnimation(IDatasmithTransformAnimationElement& TransformAnimation, const FDeltaGenTmlDataAnimationTrack& Track)
+void PopulateTransformAnimation(IDatasmithTransformAnimationElement& TransformAnimation, const FDeltaGenTmlDataAnimationTrack& Track, float Framerate)
 {
 	if (Track.Zeroed)
 	{
@@ -797,7 +797,7 @@ void PopulateTransformAnimation(IDatasmithTransformAnimationElement& TransformAn
 		Curves[2].AddKey(Key, Value.Z, bUnwindRotation);
 	}
 
-	FFrameRate FrameRate = FFrameRate(30, 1);
+	FFrameRate FrameRate = FFrameRate(static_cast<uint32>(Framerate + 0.5f), 1);
 	FFrameNumber StartFrame = FrameRate.AsFrameNumber(MinKey);
 
 	// If we use AsFrameNumber it will floor, and we might lose the very end of the animation
@@ -849,7 +849,7 @@ TSharedPtr<IDatasmithLevelSequenceElement> FDatasmithDeltaGenImporter::ConvertAn
 
 		for (const FDeltaGenTmlDataAnimationTrack& Track : Animation.Tracks)
 		{
-			PopulateTransformAnimation(TransformAnimation.Get(), Track);
+			PopulateTransformAnimation(TransformAnimation.Get(), Track, TmlTimeline.Framerate);
 		}
 
 		if (TransformAnimation->GetFramesCount(EDatasmithTransformType::Translation) > 0 ||
@@ -857,6 +857,7 @@ TSharedPtr<IDatasmithLevelSequenceElement> FDatasmithDeltaGenImporter::ConvertAn
 			TransformAnimation->GetFramesCount(EDatasmithTransformType::Scale) > 0)
 		{
 			SequenceElement->AddAnimation(TransformAnimation);
+			SequenceElement->SetFrameRate(TmlTimeline.Framerate);
 		}
 	}
 
diff --git a/Engine/Plugins/Enterprise/DatasmithFBXImporter/Source/DatasmithDeltaGenTranslator/Private/DatasmithDeltaGenImporterAuxFiles.cpp b/Engine/Plugins/Enterprise/DatasmithFBXImporter/Source/DatasmithDeltaGenTranslator/Private/DatasmithDeltaGenImporterAuxFiles.cpp
--- a/Engine/Plugins/Enterprise/DatasmithFBXImporter/Source/DatasmithDeltaGenTranslator/Private/DatasmithDeltaGenImporterAuxFiles.cpp
+++ b/Engine/Plugins/Enterprise/DatasmithFBXImporter/Source/DatasmithDeltaGenTranslator/Private/DatasmithDeltaGenImporterAuxFiles.cpp
@@ -874,8 +874,20 @@ bool LoadTMLFile(const TCHAR* InFilePath, FDatasmithDeltaGenImportTmlResult& Out
 		return false;
 	}
 
+	// Default framerate
+	float FileFramerate = 30.0f;
+
 	for(const FXmlNode* Node: FXmlNodeChildren(AnimationsNode))
 	{
+		if (TEXT("FPS") == Node->GetTag())
+		{
+			float Framerate = FCString::Atof(*Node->GetContent());
+			if (!FMath::IsNearlyZero(Framerate))
+			{
+				FileFramerate = Framerate;
+			}
+		}
+
 		if ("Timeline" != Node->GetTag())
 		{
 			continue;
@@ -899,6 +911,7 @@ bool LoadTMLFile(const TCHAR* InFilePath, FDatasmithDeltaGenImportTmlResult& Out
 		FDeltaGenTmlDataTimeline& Timeline = *(new(OutResult.Timelines) FDeltaGenTmlDataTimeline);
 
 		Timeline.Name = TimelineName;
+		Timeline.Framerate = FileFramerate;
 
 		for(const FXmlNode* AnimationNode: FXmlNodeChildren(AnimationGroupNode))
 		{
@@ -975,151 +988,201 @@ bool LoadTMLFile(const TCHAR* InFilePath, FDatasmithDeltaGenImportTmlResult& Out
 					continue;
 				}
 
+				EDeltaGenTmlDataAnimationTrackType AnimationType = EDeltaGenTmlDataAnimationTrackType::Unsupported;
+				if (TEXT("SceneObjectTranslationAnimation") == ObjectAnimationNode->GetAttribute(TEXT("Type")))
+				{
+					AnimationType = EDeltaGenTmlDataAnimationTrackType::Translation;
+				}
+				else if (TEXT("SceneObjectRotationAnimation") == ObjectAnimationNode->GetAttribute(TEXT("Type"))) // quaternion
+				{
+					AnimationType = EDeltaGenTmlDataAnimationTrackType::Rotation;
+				}
+				else if (TEXT("SceneObjectEulerAnimation") == ObjectAnimationNode->GetAttribute(TEXT("Type")))  // degrees
+				{
+					AnimationType = EDeltaGenTmlDataAnimationTrackType::RotationDeltaGenEuler;
+				}
+				else if (TEXT("SceneObjectScaleAnimation") == ObjectAnimationNode->GetAttribute(TEXT("Type")))
 				{
-					EDeltaGenTmlDataAnimationTrackType AnimationType = EDeltaGenTmlDataAnimationTrackType::Unsupported;
+					AnimationType = EDeltaGenTmlDataAnimationTrackType::Scale;
+				}
+				else if (TEXT("SceneObjectCenterAnimation") == ObjectAnimationNode->GetAttribute(TEXT("Type")))
+				{
+					AnimationType = EDeltaGenTmlDataAnimationTrackType::Center;
+				}
+
+				//FDeltaGenTmlDataAnimationTrack& AnimationTrack = TimelineAnimation.Tracks.Add(AnimationType);
+				FDeltaGenTmlDataAnimationTrack& AnimationTrack = *(new(TimelineAnimation.Tracks) FDeltaGenTmlDataAnimationTrack);
+				AnimationTrack.Type = AnimationType;
+
+				const FXmlNode* AnimationFunctionNodes = ObjectAnimationNode->FindChildNode("AnimationFunction");
+				for (const FXmlNode* AnimationFunctionNode: FXmlNodeChildren(AnimationFunctionNodes))
+				{
+					if (TEXT("Sequence") == AnimationFunctionNode->GetTag())
 					{
-						if (TEXT("SceneObjectTranslationAnimation") == ObjectAnimationNode->GetAttribute(TEXT("Type")))
+						if (TEXT("Base") == AnimationFunctionNode->GetAttribute("Role"))
 						{
-							AnimationType = EDeltaGenTmlDataAnimationTrackType::Translation;
-						}
-						else if (TEXT("SceneObjectRotationAnimation") == ObjectAnimationNode->GetAttribute(TEXT("Type"))) // quaternion
-						{
-							AnimationType = EDeltaGenTmlDataAnimationTrackType::Rotation;
-						}
-						else if (TEXT("SceneObjectEulerAnimation") == ObjectAnimationNode->GetAttribute(TEXT("Type")))  // degrees
-						{
-							AnimationType = EDeltaGenTmlDataAnimationTrackType::RotationDeltaGenEuler;
-						}
-						else if (TEXT("SceneObjectScaleAnimation") == ObjectAnimationNode->GetAttribute(TEXT("Type")))
-						{
-							AnimationType = EDeltaGenTmlDataAnimationTrackType::Scale;
-						}
-						else if (TEXT("SceneObjectCenterAnimation") == ObjectAnimationNode->GetAttribute(TEXT("Type")))
-						{
-							AnimationType = EDeltaGenTmlDataAnimationTrackType::Center;
-						}
+							const FXmlNode* KeyframesNode = AnimationFunctionNode->FindChildNode("Keyframes");
+							if (!KeyframesNode)
+							{
+								continue;
+							}
 
-					}
+							const FXmlNode* KeysNode = KeyframesNode->FindChildNode("Keys");
+							if (!KeysNode)
+							{
+								continue;
+							}
 
-					//FDeltaGenTmlDataAnimationTrack& AnimationTrack = TimelineAnimation.Tracks.Add(AnimationType);
-					FDeltaGenTmlDataAnimationTrack& AnimationTrack = *(new(TimelineAnimation.Tracks) FDeltaGenTmlDataAnimationTrack);
-					AnimationTrack.Type = AnimationType;
+							const FXmlNode* ValuesNode = KeyframesNode->FindChildNode("Values");
+							if (!ValuesNode)
+							{
+								continue;
+							}
 
-					const FXmlNode* AnimationFunctionNodes = ObjectAnimationNode->FindChildNode("AnimationFunction");
-					for (const FXmlNode* AnimationFunctionNode: FXmlNodeChildren(AnimationFunctionNodes))
-					{
-						if (TEXT("Sequence") == AnimationFunctionNode->GetTag())
-						{
-								if (TEXT("Base") == AnimationFunctionNode->GetAttribute("Role"))
+							const FString& ValuesText = ValuesNode->GetContent();
+							TArray<FString> ValuesStrings;
+							TArray<FVector4> Values;
+							if (1 < ValuesText.ParseIntoArray(ValuesStrings, TEXT(";")))
+							{
+								for(const FString& ValueString: ValuesStrings)
 								{
-									const FXmlNode* KeyframesNode = AnimationFunctionNode->FindChildNode("Keyframes");
-									if (!KeyframesNode)
+									TArray<FString> ValueTokens;
+
+									ValueString.ParseIntoArrayWS(ValueTokens);
+
+									TArray<float> ValueFloat;
+									for (int i = 0; i < ValueTokens.Num(); ++i)
 									{
-										continue;
+										ValueFloat.Add(FCString::Atof(*ValueTokens[i]));
 									}
 
-									const FXmlNode* KeysNode = KeyframesNode->FindChildNode("Keys");
-									if (!KeysNode)
+									int TockensLacking = 4-ValueTokens.Num();
+									if (0 < TockensLacking)
 									{
-										continue;
+										ValueFloat.AddZeroed(TockensLacking);
 									}
 
-									const FXmlNode* ValuesNode = KeyframesNode->FindChildNode("Values");
-									if (!ValuesNode)
+									FVector4 ValueResult(ValueFloat[0], ValueFloat[1], ValueFloat[2], ValueFloat[3]);
+									switch(AnimationType)
 									{
-										continue;
-									}
 
-									const FString& ValuesText = ValuesNode->GetContent();
-									TArray<FString> ValuesStrings;
-									TArray<FVector4> Values;
-									if (1 < ValuesText.ParseIntoArray(ValuesStrings, TEXT(";")))
+									// We export euler angles as is - because converting to quat looses information
+									// E.g. 0 vs. 360 degree euler rotation is same in quaternion representation. But for animation it makes difference.
+									// For example, take animation of [(0, 0, 0), (360, 0, 0)] in euler angles, this is full circle rotation
+									// this can't be encoded in quats as quats define orientation
+									case EDeltaGenTmlDataAnimationTrackType::RotationDeltaGenEuler:
 									{
-										for(const FString& ValueString: ValuesStrings)
+										FVector Euler(ValueFloat[0], ValueFloat[1], ValueFloat[2]);
+										ValueResult = FVector4(Euler, 0.0f);
+									}
+									break;
+									case EDeltaGenTmlDataAnimationTrackType::Rotation:
 										{
-											TArray<FString> ValueTokens;
-
-											ValueString.ParseIntoArrayWS(ValueTokens);
-
-											TArray<float> ValueFloat;
-											for (int i = 0; i < ValueTokens.Num(); ++i)
-											{
-												ValueFloat.Add(FCString::Atof(*ValueTokens[i]));
-											}
-
-											int TockensLacking = 4-ValueTokens.Num();
-											if (0 < TockensLacking)
-											{
-												ValueFloat.AddZeroed(TockensLacking);
-											}
-
-											FVector4 ValueResult(ValueFloat[0], ValueFloat[1], ValueFloat[2], ValueFloat[3]);
-											switch(AnimationType)
-											{
-
-											// We export euler angles as is - because converting to quat looses information
-											// E.g. 0 vs. 360 degree euler rotation is same in quaternion representation. But for animation it makes difference.
-											// For example, take animation of [(0, 0, 0), (360, 0, 0)] in euler angles, this is full circle rotation
-											// this can't be encoded in quats as quats define orientation
-											case EDeltaGenTmlDataAnimationTrackType::RotationDeltaGenEuler:
-											{
-												FVector Euler(ValueFloat[0], ValueFloat[1], ValueFloat[2]);
-												ValueResult = FVector4(Euler, 0.0f);
-											}
-											break;
-											case EDeltaGenTmlDataAnimationTrackType::Rotation:
-												{
-													// taked from FFbxDataConverter::ConvertRotToQuat
-													FQuat Quat(ValueFloat[0], -ValueFloat[1], ValueFloat[2], -ValueFloat[3]);
-													auto Rotator = Quat.Rotator();
-													ValueResult = FVector4(Rotator.Euler(), 0.0f);
-												}
-												break;
-											};
-
-											Values.Add(ValueResult);
+											// taked from FFbxDataConverter::ConvertRotToQuat
+											FQuat Quat(ValueFloat[0], -ValueFloat[1], ValueFloat[2], -ValueFloat[3]);
+											auto Rotator = Quat.Rotator();
+											ValueResult = FVector4(Rotator.Euler(), 0.0f);
 										}
-									}
+										break;
+									};
 
-									const FString& KeysText = KeysNode->GetContent();
+									Values.Add(ValueResult);
+								}
+							}
 
-									TArray<FString> KeysStrings;
-									TArray<float> Keys;
-									if (1 < KeysText.ParseIntoArray(KeysStrings, TEXT(";")))
-									{
-										for(const FString& KeyString: KeysStrings)
-										{
-											Keys.Add(FCString::Atof(*KeyString));
-										}
-									}
+							const FString& KeysText = KeysNode->GetContent();
 
-									if (Keys.Num() == Values.Num())
-									{
-										for (int i=0;i!=Values.Num();++i)
-										{
-											AnimationTrack.Keys.Add(Keys[i]);
-											AnimationTrack.Values.Add(Values[i]);
-										}
-									}
+							TArray<FString> KeysStrings;
+							TArray<float> Keys;
+							if (1 < KeysText.ParseIntoArray(KeysStrings, TEXT(";")))
+							{
+								for(const FString& KeyString: KeysStrings)
+								{
+									Keys.Add(FCString::Atof(*KeyString));
 								}
-								else if ("TimeAdjustment" == AnimationFunctionNode->GetAttribute("Role"))
+							}
+
+							if (Keys.Num() == Values.Num())
+							{
+								for (int i=0;i!=Values.Num();++i)
 								{
+									AnimationTrack.Keys.Add(Keys[i]);
+									AnimationTrack.Values.Add(Values[i]);
 								}
+							}
 						}
 					}
 
-					// all zeroes seems to be the same as absent animation
-					bool AllZeroes = true;
-					for(auto Value: AnimationTrack.Values)
+					// This is a strange behavior from DeltaGen: If the animation was originally 25fps and it was later
+					// converted to 50fps, the .tml file will still contain '25' for the FPS, but the animations will have
+					// a "TimeAdjustment" Interpolator with control vertices that have an extra multiplier in the x dimension.
+					// Example: In this case, the Interpolator's ControlVertices xml element would look like this:
+					//						<Positions>
+					//							0.00000000 0.00000000 0.00000000;
+					//							6.66666651 3.33333325 0.00000000;
+					//							13.33333302 6.66666651 0.00000000;
+					//							20.00000000 10.00000000 0.00000000
+					//						</Positions>
+					// Notice how x * 2 = y: This means that 2 is our factor, so the animation framerate is 25 * 2 = 50fps
+					// I have no idea why this is done, or if it actually works like this every time, but that's as much info
+					// as we have
+					else if (TEXT("Interpolator") == AnimationFunctionNode->GetTag())
 					{
-						if (Value != FVector4(0.0f, 0.0f, 0.0f, 0.0f))
+						if (TEXT("TimeAdjustment") == AnimationFunctionNode->GetAttribute(TEXT("Role")))
 						{
-							AllZeroes = false;
-							break;
+							const FXmlNode* InterpolationCurveNode = AnimationFunctionNode->FindChildNode(TEXT("InterpolationCurve"));
+							if (!InterpolationCurveNode)
+							{
+								continue;
+							}
+
+							const FXmlNode* ControlVerticesNode = InterpolationCurveNode->FindChildNode(TEXT("ControlVertices"));
+							if (!ControlVerticesNode)
+							{
+								continue;
+							}
+
+							const FXmlNode* PositionsNode = ControlVerticesNode->FindChildNode(TEXT("Positions"));
+							if (!PositionsNode)
+							{
+								continue;
+							}
+
+							const FString& PositionsText = PositionsNode->GetContent();
+							TArray<FString> VectorStrings;
+							if (1 < PositionsText.ParseIntoArray(VectorStrings, TEXT(";")))
+							{
+								for (const FString& VectorString : VectorStrings)
+								{
+									TArray<FString> ValueStrings;
+									if (2 <= VectorString.ParseIntoArray(ValueStrings, TEXT(" ")))
+									{
+										float X = FCString::Atof(*ValueStrings[0]);
+										float Y = FCString::Atof(*ValueStrings[1]);
+
+										if(!FMath::IsNearlyZero(Y))
+										{
+											Timeline.Framerate = FMath::Max(Timeline.Framerate, FileFramerate * X / Y);
+											break;
+										}
+									}
+								}
+							}
 						}
 					}
-					AnimationTrack.Zeroed = AllZeroes;
 				}
+
+				// all zeroes seems to be the same as absent animation
+				bool AllZeroes = true;
+				for(auto Value: AnimationTrack.Values)
+				{
+					if (Value != FVector4(0.0f, 0.0f, 0.0f, 0.0f))
+					{
+						AllZeroes = false;
+						break;
+					}
+				}
+				AnimationTrack.Zeroed = AllZeroes;
 			}
 		}
 	}
diff --git a/Engine/Plugins/Enterprise/DatasmithFBXImporter/Source/DatasmithFBXTranslator/Private/DatasmithFBXScene.cpp b/Engine/Plugins/Enterprise/DatasmithFBXImporter/Source/DatasmithFBXTranslator/Private/DatasmithFBXScene.cpp
--- a/Engine/Plugins/Enterprise/DatasmithFBXImporter/Source/DatasmithFBXTranslator/Private/DatasmithFBXScene.cpp
+++ b/Engine/Plugins/Enterprise/DatasmithFBXImporter/Source/DatasmithFBXTranslator/Private/DatasmithFBXScene.cpp
@@ -148,7 +148,10 @@ void FDatasmithFBXSceneNode::MarkSwitchNode()
 	NodeType |= ENodeType::Switch;
 	bShouldKeepThisNode = true;
 
-	// For switch, we should also keep all its children persistent because they are representing variants
+	// Prevent the targets of SwitchObjects from being removed. This shouldn't be necessary,
+	// but some scenes seem to have all mesh nodes named exactly the same (e.g shell_0).
+	// In those cases, when we convert them to Switch variants, we will search for actors named
+	// 'shell_0' for our switch options and likely end up picking the wrong actor
 	for (int32 NodeIndex = 0; NodeIndex < Children.Num(); NodeIndex++)
 	{
 		Children[NodeIndex]->bShouldKeepThisNode = true;
diff --git a/Engine/Plugins/Enterprise/DatasmithFBXImporter/Source/DatasmithFBXTranslator/Private/DatasmithFBXSceneProcessor.cpp b/Engine/Plugins/Enterprise/DatasmithFBXImporter/Source/DatasmithFBXTranslator/Private/DatasmithFBXSceneProcessor.cpp
--- a/Engine/Plugins/Enterprise/DatasmithFBXImporter/Source/DatasmithFBXTranslator/Private/DatasmithFBXSceneProcessor.cpp
+++ b/Engine/Plugins/Enterprise/DatasmithFBXImporter/Source/DatasmithFBXTranslator/Private/DatasmithFBXSceneProcessor.cpp
@@ -213,7 +213,7 @@ struct FNodeMarkHelper
 void FDatasmithFBXSceneProcessor::FindPersistentNodes()
 {
 	FNodeMarkHelper Helper;
-	Helper.SwitchMaterialObjectNames.Append(Scene->SwitchObjects);
+	Helper.SwitchObjectNames.Append(Scene->SwitchObjects);
 	Helper.AnimatedObjectNames.Append(Scene->AnimatedObjects);
 	Helper.SwitchMaterialObjectNames.Append(Scene->SwitchMaterialObjects);
 	Helper.TransformVariantObjectNames.Append(Scene->TransformVariantObjects);
diff --git a/Engine/Plugins/Enterprise/DatasmithFBXImporter/Source/DatasmithVREDTranslator/Private/DatasmithVREDImporter.cpp b/Engine/Plugins/Enterprise/DatasmithFBXImporter/Source/DatasmithVREDTranslator/Private/DatasmithVREDImporter.cpp
--- a/Engine/Plugins/Enterprise/DatasmithFBXImporter/Source/DatasmithVREDTranslator/Private/DatasmithVREDImporter.cpp
+++ b/Engine/Plugins/Enterprise/DatasmithFBXImporter/Source/DatasmithVREDTranslator/Private/DatasmithVREDImporter.cpp
@@ -1320,7 +1320,7 @@ TSharedPtr<IDatasmithBaseMaterialElement> FDatasmithVREDImporter::ConvertMateria
 	return MaterialElement;
 }
 
-void PopulateTransformAnimation(IDatasmithTransformAnimationElement& TransformAnimation, EDatasmithTransformType AnimationType, const FDatasmithFBXSceneAnimCurve* CurveX, const FDatasmithFBXSceneAnimCurve* CurveY, const FDatasmithFBXSceneAnimCurve* CurveZ, double ScaleFactor)
+void PopulateTransformAnimation(IDatasmithTransformAnimationElement& TransformAnimation, EDatasmithTransformType AnimationType, const FDatasmithFBXSceneAnimCurve* CurveX, const FDatasmithFBXSceneAnimCurve* CurveY, const FDatasmithFBXSceneAnimCurve* CurveZ, double ScaleFactor, float Framerate)
 {
 	if (!CurveX && !CurveY && !CurveZ)
 	{
@@ -1425,7 +1425,7 @@ void PopulateTransformAnimation(IDatasmithTransformAnimationElement& TransformAn
 
 	TransformAnimation.SetEnabledTransformChannels(Channels | FDatasmithAnimationUtils::SetChannelTypeComponents(Components, AnimationType));
 
-	FFrameRate FrameRate = FFrameRate(30, 1); // The default for LevelSequenceElements
+	FFrameRate FrameRate = FFrameRate(static_cast<uint32>(Framerate + 0.5f), 1);
 	FFrameNumber StartFrame = FrameRate.AsFrameNumber(MinKey);
 	FFrameNumber EndFrame = FrameRate.AsFrameTime(MaxKey).CeilToFrame();
 
@@ -1464,7 +1464,7 @@ void PopulateTransformAnimation(IDatasmithTransformAnimationElement& TransformAn
 	}
 }
 
-void PopulateVisibilityAnimation(IDatasmithVisibilityAnimationElement& VisibilityAnimation, const FDatasmithFBXSceneAnimCurve* Curve)
+void PopulateVisibilityAnimation(IDatasmithVisibilityAnimationElement& VisibilityAnimation, const FDatasmithFBXSceneAnimCurve* Curve, float Framerate)
 {
 	if (Curve == nullptr || Curve->Points.Num() == 0)
 	{
@@ -1491,7 +1491,7 @@ void PopulateVisibilityAnimation(IDatasmithVisibilityAnimationElement& Visibilit
 		MaxKey = FMath::Max(MaxKey, CurvePoint->Time);
 	}
 
-	FFrameRate FrameRate = FFrameRate(30, 1);
+	FFrameRate FrameRate = FFrameRate(static_cast<uint32>(Framerate + 0.5f), 1);
 	FFrameNumber StartFrame = FrameRate.AsFrameNumber(MinKey);
 	FFrameNumber EndFrame = FrameRate.AsFrameTime(MaxKey).CeilToFrame();
 
@@ -1552,6 +1552,7 @@ TSharedPtr<IDatasmithLevelSequenceElement> FDatasmithVREDImporter::ConvertAnimBl
 	FString BlockName = CombinedBlock.NodeNameToBlock.CreateConstIterator().Value()->Name;
 
 	TSharedRef<IDatasmithLevelSequenceElement> SequenceElement = FDatasmithSceneFactory::CreateLevelSequence(*BlockName);
+	SequenceElement->SetFrameRate(IntermediateScene->BaseTime);
 
 	for (const auto& Pair : CombinedBlock.NodeNameToBlock)
 	{
@@ -1586,7 +1587,7 @@ TSharedPtr<IDatasmithLevelSequenceElement> FDatasmithVREDImporter::ConvertAnimBl
 
 				if (TransformAnimation.IsValid())
 				{
-					PopulateTransformAnimation(TransformAnimation.ToSharedRef().Get(), EDatasmithTransformType::Translation, CurveX, CurveY, CurveZ, IntermediateScene->ScaleFactor);
+					PopulateTransformAnimation(TransformAnimation.ToSharedRef().Get(), EDatasmithTransformType::Translation, CurveX, CurveY, CurveZ, IntermediateScene->ScaleFactor, IntermediateScene->BaseTime);
 				}
 			}
 		}
@@ -1608,7 +1609,7 @@ TSharedPtr<IDatasmithLevelSequenceElement> FDatasmithVREDImporter::ConvertAnimBl
 
 				if (TransformAnimation.IsValid())
 				{
-					PopulateTransformAnimation(TransformAnimation.ToSharedRef().Get(), EDatasmithTransformType::Scale, CurveX, CurveY, CurveZ, IntermediateScene->ScaleFactor);
+					PopulateTransformAnimation(TransformAnimation.ToSharedRef().Get(), EDatasmithTransformType::Scale, CurveX, CurveY, CurveZ, IntermediateScene->ScaleFactor, IntermediateScene->BaseTime);
 				}
 			}
 		}
@@ -1630,7 +1631,7 @@ TSharedPtr<IDatasmithLevelSequenceElement> FDatasmithVREDImporter::ConvertAnimBl
 
 				if (TransformAnimation.IsValid())
 				{
-					PopulateTransformAnimation(TransformAnimation.ToSharedRef().Get(), EDatasmithTransformType::Rotation, CurveX, CurveY, CurveZ, IntermediateScene->ScaleFactor);
+					PopulateTransformAnimation(TransformAnimation.ToSharedRef().Get(), EDatasmithTransformType::Rotation, CurveX, CurveY, CurveZ, IntermediateScene->ScaleFactor, IntermediateScene->BaseTime);
 				}
 			}
 		}
@@ -1653,7 +1654,7 @@ TSharedPtr<IDatasmithLevelSequenceElement> FDatasmithVREDImporter::ConvertAnimBl
 				if (VisibilityAnimation.IsValid())
 				{
 					VisibilityAnimation->SetPropagateToChildren(true);
-					PopulateVisibilityAnimation(VisibilityAnimation.ToSharedRef().Get(), Curves[0]);
+					PopulateVisibilityAnimation(VisibilityAnimation.ToSharedRef().Get(), Curves[0], IntermediateScene->BaseTime);
 				}
 			}
 		}
@@ -1674,7 +1675,7 @@ TSharedPtr<IDatasmithLevelSequenceElement> FDatasmithVREDImporter::ConvertAnimBl
 	return SequenceElement;
 }
 
-void ConvertAnimClips(TArray<FDatasmithFBXSceneAnimClip>& AnimClips, TMap<FString, FImportedAnim>& ImportedAnims, TArray<TSharedPtr<IDatasmithLevelSequenceElement>>& OutCreatedClips)
+void ConvertAnimClips(TArray<FDatasmithFBXSceneAnimClip>& AnimClips, TMap<FString, FImportedAnim>& ImportedAnims, TArray<TSharedPtr<IDatasmithLevelSequenceElement>>& OutCreatedClips, float Framerate)
 {
 	TArray<FDatasmithFBXSceneAnimClip> ClipsRemaining = AnimClips;
 
@@ -1717,8 +1718,9 @@ void ConvertAnimClips(TArray<FDatasmithFBXSceneAnimClip>& AnimClips, TMap<FStrin
 
 			// Create a IDatasmithLevelSequenceElement and populate it with pointers to the target subsequence elements
 			TSharedPtr<IDatasmithLevelSequenceElement> SequenceElement = FDatasmithSceneFactory::CreateLevelSequence(*Clip.Name);
+			SequenceElement->SetFrameRate(Framerate);
 
-			FFrameRate FrameRate = FFrameRate(FMath::RoundToInt(SequenceElement->GetFrameRate()), 1);
+			FFrameRate FrameRate = FFrameRate(static_cast<uint32>(Framerate + 0.5f), 1);
 
 			float MinStartTime = FLT_MAX;
 			float MaxEndTime = -FLT_MAX;
@@ -1964,7 +1966,7 @@ bool FDatasmithVREDImporter::SendSceneToDatasmith()
 		}
 
 		TArray<TSharedPtr<IDatasmithLevelSequenceElement>> CreatedClips;
-		ConvertAnimClips(ParsedAnimClips, ImportedAnimElements, CreatedClips);
+		ConvertAnimClips(ParsedAnimClips, ImportedAnimElements, CreatedClips, IntermediateScene->BaseTime);
 
 		for (const TSharedPtr<IDatasmithLevelSequenceElement>& CreatedClip : CreatedClips)
 		{
diff --git a/Engine/Plugins/Enterprise/DatasmithFBXImporter/Source/DatasmithVREDTranslator/Private/DatasmithVREDImporterMaterialSelector.cpp b/Engine/Plugins/Enterprise/DatasmithFBXImporter/Source/DatasmithVREDTranslator/Private/DatasmithVREDImporterMaterialSelector.cpp
--- a/Engine/Plugins/Enterprise/DatasmithFBXImporter/Source/DatasmithVREDTranslator/Private/DatasmithVREDImporterMaterialSelector.cpp
+++ b/Engine/Plugins/Enterprise/DatasmithFBXImporter/Source/DatasmithVREDTranslator/Private/DatasmithVREDImporterMaterialSelector.cpp
@@ -88,12 +88,8 @@ void FDatasmithVREDImporterMaterialSelector::FinalizeMaterialInstance(const TSha
 		((OpacityProperty.IsValid() && GetFloat(OpacityProperty, Opacity) && Opacity < 1.0f) ||
 		(TransparencyTextureProperty.IsValid() && GetBool(TransparencyTextureProperty, bTransparencyActive) && bTransparencyActive)))
 	{
-		// Commented out due to this strange bug where if we enable these overrides, having either a transparency texture or a
-		// bump texture assigned to the material (even if disabled) will cause it to crash. It was either disabling this functionality
-		// or completely discarding the imported information on those textures. At least like this the user can manually switch the blend mode
-		// override later with no consequence
-		//MaterialInstance->BasePropertyOverrides.bOverride_BlendMode = true;
-		//MaterialInstance->BasePropertyOverrides.BlendMode = EBlendMode::BLEND_Translucent;
+		MaterialInstance->BasePropertyOverrides.bOverride_BlendMode = true;
+		MaterialInstance->BasePropertyOverrides.BlendMode = EBlendMode::BLEND_Translucent;
 	}
 }
 
diff --git a/Engine/Plugins/Enterprise/DatasmithGLTFImporter/Source/DatasmithGLTFTranslator/Private/DatasmithGLTFImporter.cpp b/Engine/Plugins/Enterprise/DatasmithGLTFImporter/Source/DatasmithGLTFTranslator/Private/DatasmithGLTFImporter.cpp
--- a/Engine/Plugins/Enterprise/DatasmithGLTFImporter/Source/DatasmithGLTFTranslator/Private/DatasmithGLTFImporter.cpp
+++ b/Engine/Plugins/Enterprise/DatasmithGLTFImporter/Source/DatasmithGLTFTranslator/Private/DatasmithGLTFImporter.cpp
@@ -10,6 +10,7 @@
 #include "GLTFAsset.h"
 #include "GLTFMaterialFactory.h"
 #include "GLTFReader.h"
+#include "GLTFMesh.h"
 #include "GLTFMeshFactory.h"
 
 #include "DatasmithAssetImportData.h"
@@ -22,6 +23,7 @@
 #include "AssetRegistryModule.h"
 #include "Engine/StaticMesh.h"
 #include "Misc/Paths.h"
+#include "Misc/SecureHash.h"
 
 DEFINE_LOG_CATEGORY(LogDatasmithGLTFImport);
 
@@ -188,13 +190,23 @@ TSharedPtr<IDatasmithMeshActorElement> FDatasmithGLTFImporter::CreateStaticMeshA
 	{
 		ImportedMeshes.Add(MeshIndex);
 
-		TSharedRef<IDatasmithMeshElement> MeshElement = FDatasmithSceneFactory::CreateMesh(*GLTFAsset->Meshes[MeshIndex].Name);
+		const GLTF::FMesh& Mesh = GLTFAsset->Meshes[MeshIndex];
+
+		TSharedRef<IDatasmithMeshElement> MeshElement = FDatasmithSceneFactory::CreateMesh(*Mesh.Name);
+
+		FMD5 MD5;
+
+		FMD5Hash MeshHash = Mesh.GetHash();
+		MD5.Update(MeshHash.GetBytes(), MeshHash.GetSize());
 
 		const TArray<GLTF::FMaterialElement*>& Materials = MaterialFactory->GetMaterials();
 		for (int32 MaterialID = 0; MaterialID < Materials.Num(); ++MaterialID)
 		{
 			GLTF::FMaterialElement* Material = Materials[MaterialID];
 			MeshElement->SetMaterial(*Material->GetName(), MaterialID);
+
+			FMD5Hash MaterialHash = Material->GetGLTFMaterialHash();
+			MD5.Update(MaterialHash.GetBytes(), MaterialHash.GetSize());
 		}
 
 		if (ImportOptions->bGenerateLightmapUVs)
@@ -207,6 +219,13 @@ TSharedPtr<IDatasmithMeshActorElement> FDatasmithGLTFImporter::CreateStaticMeshA
 			MeshElement->SetLightmapCoordinateIndex(0);
 		}
 
+		uint8 GenerateLightmapUVs = static_cast<uint8>(ImportOptions->bGenerateLightmapUVs);
+		MD5.Update(&GenerateLightmapUVs, 1);
+
+		FMD5Hash Result;
+		Result.Set(MD5);
+		MeshElement->SetFileHash(Result);
+
 		MeshElementToGLTFMeshIndex.Add(&MeshElement.Get(), MeshIndex);
 		GLTFMeshIndexToMeshElement.Add(MeshIndex, MeshElement);
 		DatasmithScene->AddMesh(MeshElement);
diff --git a/Engine/Plugins/Enterprise/DatasmithGLTFImporter/Source/DatasmithGLTFTranslator/Private/DatasmithGLTFTextureFactory.cpp b/Engine/Plugins/Enterprise/DatasmithGLTFImporter/Source/DatasmithGLTFTranslator/Private/DatasmithGLTFTextureFactory.cpp
--- a/Engine/Plugins/Enterprise/DatasmithGLTFImporter/Source/DatasmithGLTFTranslator/Private/DatasmithGLTFTextureFactory.cpp
+++ b/Engine/Plugins/Enterprise/DatasmithGLTFImporter/Source/DatasmithGLTFTranslator/Private/DatasmithGLTFTextureFactory.cpp
@@ -93,10 +93,14 @@ GLTF::ITextureElement* FDatasmithGLTFTextureFactory::CreateTexture(const GLTF::F
 	using namespace DatasmithGLTFImporterImpl;
 
 	if (GltfTexture.Name.IsEmpty())
+	{
 		return nullptr;
+	}
 
 	if (GltfTexture.Source.FilePath.IsEmpty() && GltfTexture.Source.DataByteLength == 0)
+	{
 		return nullptr;
+	}
 
 	FString TextureName = GltfTexture.Name;
 	if (CreatedTextures.Contains(TextureName))
@@ -109,13 +113,16 @@ GLTF::ITextureElement* FDatasmithGLTFTextureFactory::CreateTexture(const GLTF::F
 	Texture->SetTextureFilter(ConvertFilter(GltfTexture.Sampler.MinFilter));
 	Texture->SetTextureAddressX(ConvertWrap(GltfTexture.Sampler.WrapS));
 	Texture->SetTextureAddressY(ConvertWrap(GltfTexture.Sampler.WrapT));
-	if (TextureMode != GLTF::ETextureMode::Color)
-		Texture->SetRGBCurve(1.f);
+	Texture->SetSRGB(TextureMode == GLTF::ETextureMode::Color ? EDatasmithColorSpace::sRGB : EDatasmithColorSpace::Linear);
 
 	if (GltfTexture.Source.DataByteLength > 0)
+	{
 		Texture->SetData(GltfTexture.Source.Data, GltfTexture.Source.DataByteLength, ConvertFormat(GltfTexture.Source.Format));
+	}
 	else
+	{
 		Texture->SetFile(*GltfTexture.Source.FilePath);
+	}
 
 	CurrentScene->AddTexture(Texture);
 
diff --git a/Engine/Plugins/Enterprise/DatasmithIFCImporter/Source/DatasmithIFCTranslator/DatasmithIFCTranslator.Build.cs b/Engine/Plugins/Enterprise/DatasmithIFCImporter/Source/DatasmithIFCTranslator/DatasmithIFCTranslator.Build.cs
--- a/Engine/Plugins/Enterprise/DatasmithIFCImporter/Source/DatasmithIFCTranslator/DatasmithIFCTranslator.Build.cs
+++ b/Engine/Plugins/Enterprise/DatasmithIFCImporter/Source/DatasmithIFCTranslator/DatasmithIFCTranslator.Build.cs
@@ -14,14 +14,10 @@ namespace UnrealBuildTool.Rules
                 new string[]
                 {
                     "Analytics",
-                    "Core",
                     "CoreUObject",
                     "DatasmithCore",
 					"Engine",
 					"Json",
-                    "MainFrame",
-					"MaterialEditor",
-					"MeshUtilities",
                     "MeshDescription",
                     "MessageLog",
 					"Projects",
@@ -29,7 +25,6 @@ namespace UnrealBuildTool.Rules
                     "Slate",
                     "SlateCore",
 					"StaticMeshDescription",
-                    "UnrealEd",
                     "XmlParser",
                 }
             );
@@ -48,14 +43,16 @@ namespace UnrealBuildTool.Rules
 			PublicDependencyModuleNames.AddRange(
                 new string[]
                 {
+                    "Core",
                     "DatasmithContent",
-                    "DatasmithImporter"
-                }
+					"DatasmithImporter"
+				}
             );
 
 			if (Target.Platform == UnrealTargetPlatform.Win64)
 			{
 				RuntimeDependencies.Add("$(EngineDir)/Plugins/Enterprise/DatasmithIFCImporter/Binaries/Win64/ifcengine.dll");
+				PublicDelayLoadDLLs.Add("ifcengine.dll");
 			}
 		}
 	}
diff --git a/Engine/Plugins/Enterprise/DatasmithIFCImporter/Source/DatasmithIFCTranslator/Private/DatasmithIFCImporter.cpp b/Engine/Plugins/Enterprise/DatasmithIFCImporter/Source/DatasmithIFCTranslator/Private/DatasmithIFCImporter.cpp
--- a/Engine/Plugins/Enterprise/DatasmithIFCImporter/Source/DatasmithIFCTranslator/Private/DatasmithIFCImporter.cpp
+++ b/Engine/Plugins/Enterprise/DatasmithIFCImporter/Source/DatasmithIFCTranslator/Private/DatasmithIFCImporter.cpp
@@ -7,19 +7,17 @@
 #include "IFC/IFCStaticMeshFactory.h"
 
 #include "DatasmithAssetImportData.h"
-#include "DatasmithImportContext.h"
 #include "DatasmithImportOptions.h"
 #include "DatasmithMeshHelper.h"
 #include "DatasmithSceneFactory.h"
+#include "DatasmithUtils.h"
 #include "IDatasmithSceneElements.h"
 #include "ObjectTemplates/DatasmithStaticMeshTemplate.h"
-#include "Utility/DatasmithImporterUtils.h"
 
 #include "AnalyticsEventAttribute.h"
 #include "AssetRegistryModule.h"
 #include "Engine/StaticMesh.h"
 #include "Misc/Paths.h"
-#include "ObjectTools.h"
 
 DEFINE_LOG_CATEGORY(LogDatasmithIFCImport);
 
@@ -64,7 +62,7 @@ bool FDatasmithIFCImporter::OpenFile(const FString& InFileName)
 
 FString FDatasmithIFCImporter::GetFilteredName(const FString& Name)
 {
-	return ObjectTools::SanitizeObjectName(Name);
+	return FDatasmithUtils::SanitizeObjectName(Name);
 }
 
 FString FDatasmithIFCImporter::GetUniqueName(const FString& Name)
diff --git a/Engine/Plugins/Enterprise/DatasmithIFCImporter/Source/DatasmithIFCTranslator/Private/DatasmithIFCTranslatorModule.cpp b/Engine/Plugins/Enterprise/DatasmithIFCImporter/Source/DatasmithIFCTranslator/Private/DatasmithIFCTranslatorModule.cpp
--- a/Engine/Plugins/Enterprise/DatasmithIFCImporter/Source/DatasmithIFCTranslator/Private/DatasmithIFCTranslatorModule.cpp
+++ b/Engine/Plugins/Enterprise/DatasmithIFCImporter/Source/DatasmithIFCTranslator/Private/DatasmithIFCTranslatorModule.cpp
@@ -6,6 +6,10 @@
 #include "Translators/DatasmithTranslator.h"
 #include "DatasmithImporterModule.h"
 #include "CoreMinimal.h"
+#include "Interfaces/IPluginManager.h"
+#include "Misc/Paths.h"
+
+DEFINE_LOG_CATEGORY(LogIFCTranslator);
 
 const TCHAR* IDatasmithIFCTranslatorModule::ModuleName = TEXT("DatasmithIFCTranslator");
 
@@ -14,6 +18,31 @@ class FIFCTranslatorModule : public IDatasmithIFCTranslatorModule
 public:
 	virtual void StartupModule() override
 	{
+		check(LibHandle == nullptr);
+
+#ifdef WITH_IFC_ENGINE_LIB
+		const FString IFCEngineDll = TEXT("ifcengine.dll");
+
+		// determine directory paths
+		FString IFCEngineDllPath = FPaths::Combine(IPluginManager::Get().FindPlugin(TEXT("DatasmithIFCImporter"))->GetBaseDir(), TEXT("Binaries"), FPlatformProcess::GetBinariesSubdirectory());
+		FPlatformProcess::PushDllDirectory(*IFCEngineDllPath);
+		IFCEngineDllPath = FPaths::Combine(IFCEngineDllPath, IFCEngineDll);
+
+		if (!FPaths::FileExists(IFCEngineDllPath))
+		{
+			UE_LOG(LogIFCTranslator, Error, TEXT("Failed to find the binary folder for the IFCEngine dll. Plug-in will not be functional."));
+			return;
+		}
+
+		LibHandle = FPlatformProcess::GetDllHandle(*IFCEngineDllPath);
+
+		if (LibHandle == nullptr)
+		{
+			UE_LOG(LogIFCTranslator, Error, TEXT("Failed to load required library %s. Plug-in will not be functional."), *IFCEngineDllPath);
+			return;
+		}
+#endif
+
 		FModuleManager::Get().LoadModule(TEXT("DatasmithImporter"));
 		Datasmith::RegisterTranslator<FDatasmithIFCTranslator>();
 	}
@@ -21,7 +50,19 @@ public:
 	virtual void ShutdownModule() override
 	{
 		Datasmith::UnregisterTranslator<FDatasmithIFCTranslator>();
+
+		if (LibHandle != nullptr)
+		{
+			FPlatformProcess::FreeDllHandle(LibHandle);
+			LibHandle = nullptr;
+		}
 	}
+
+private:
+
+	static void* LibHandle;
 };
 
+void* FIFCTranslatorModule::LibHandle = nullptr;
+
 IMPLEMENT_MODULE(FIFCTranslatorModule, DatasmithIFCTranslator);
diff --git a/Engine/Plugins/Enterprise/DatasmithIFCImporter/Source/DatasmithIFCTranslator/Private/IFC/IFCStaticMeshFactory.cpp b/Engine/Plugins/Enterprise/DatasmithIFCImporter/Source/DatasmithIFCTranslator/Private/IFC/IFCStaticMeshFactory.cpp
--- a/Engine/Plugins/Enterprise/DatasmithIFCImporter/Source/DatasmithIFCTranslator/Private/IFC/IFCStaticMeshFactory.cpp
+++ b/Engine/Plugins/Enterprise/DatasmithIFCImporter/Source/DatasmithIFCTranslator/Private/IFC/IFCStaticMeshFactory.cpp
@@ -7,7 +7,6 @@
 #include "Engine/StaticMesh.h"
 #include "StaticMeshAttributes.h"
 #include "MeshDescription.h"
-#include "PackageTools.h"
 
 DEFINE_LOG_CATEGORY(LogDatasmithIFCMeshFactory);
 
@@ -42,32 +41,26 @@ namespace IFC
 
 	void FStaticMeshFactory::FillMeshDescription(const IFC::FObject* InObject, FMeshDescription* MeshDescription) const
 	{
-		const int32 NumUVs = 1;
-
 		TVertexAttributesRef<FVector> VertexPositions = MeshDescription->VertexAttributes().GetAttributesRef<FVector>(MeshAttribute::Vertex::Position);
-
-		TEdgeAttributesRef<bool>  EdgeHardnesses = MeshDescription->EdgeAttributes().GetAttributesRef<bool>(MeshAttribute::Edge::IsHard);
-		TEdgeAttributesRef<float> EdgeCreaseSharpnesses = MeshDescription->EdgeAttributes().GetAttributesRef<float>(MeshAttribute::Edge::CreaseSharpness);
 		TPolygonGroupAttributesRef<FName> PolygonGroupImportedMaterialSlotNames = MeshDescription->PolygonGroupAttributes().GetAttributesRef<FName>(MeshAttribute::PolygonGroup::ImportedMaterialSlotName);
 		TVertexInstanceAttributesRef<FVector> VertexInstanceNormals = MeshDescription->VertexInstanceAttributes().GetAttributesRef<FVector>(MeshAttribute::VertexInstance::Normal);
-		TVertexInstanceAttributesRef<FVector> VertexInstanceTangents = MeshDescription->VertexInstanceAttributes().GetAttributesRef<FVector>(MeshAttribute::VertexInstance::Tangent);
-		TVertexInstanceAttributesRef<float> VertexInstanceBinormalSigns = MeshDescription->VertexInstanceAttributes().GetAttributesRef<float>(MeshAttribute::VertexInstance::BinormalSign);
 		TVertexInstanceAttributesRef<FVector2D> VertexInstanceUVs = MeshDescription->VertexInstanceAttributes().GetAttributesRef<FVector2D>(MeshAttribute::VertexInstance::TextureCoordinate);
-		TVertexInstanceAttributesRef<FVector4> VertexInstanceColors = MeshDescription->VertexInstanceAttributes().GetAttributesRef<FVector4>(MeshAttribute::VertexInstance::Color);
+
+		const int32 NumUVs = 1;
 		VertexInstanceUVs.SetNumIndices(NumUVs);
 
 		FIndexVertexIdMap PositionIndexToVertexId;
 		PositionIndexToVertexId.Empty(InObject->facesVerticesCount);
 
-		for (int64 I = 0; I < InObject->facesVerticesCount; I++)
+		for (int64 IFCVertexIndex = 0; IFCVertexIndex < InObject->facesVerticesCount; IFCVertexIndex++)
 		{
 			const FVertexID& VertexID = MeshDescription->CreateVertex();
-			VertexPositions[VertexID].X = InObject->facesVertices[(I * (InObject->vertexElementSize / sizeof(float))) + 0];
+			VertexPositions[VertexID].X = InObject->facesVertices[(IFCVertexIndex * (InObject->vertexElementSize / sizeof(float))) + 0];
 			// Flip Y to keep mesh looking the same as the coordinate system changes from RH -> LH
-			VertexPositions[VertexID].Y = - InObject->facesVertices[(I * (InObject->vertexElementSize / sizeof(float))) + 1];
-			VertexPositions[VertexID].Z = InObject->facesVertices[(I * (InObject->vertexElementSize / sizeof(float))) + 2];
+			VertexPositions[VertexID].Y = - InObject->facesVertices[(IFCVertexIndex * (InObject->vertexElementSize / sizeof(float))) + 1];
+			VertexPositions[VertexID].Z = InObject->facesVertices[(IFCVertexIndex * (InObject->vertexElementSize / sizeof(float))) + 2];
 			VertexPositions[VertexID] *= ImportUniformScale;
-			PositionIndexToVertexId.Add(I, VertexID);
+			PositionIndexToVertexId.Add(IFCVertexIndex, VertexID);
 		}
 
 		// Add the PolygonGroups.
@@ -81,54 +74,56 @@ namespace IFC
 			PolygonGroupImportedMaterialSlotNames[PolygonGroupID] = ImportedSlotName;
 		}
 
-		bool bMeshUsesEmptyMaterial = false;
-		bool bDidGenerateTexCoords = false;
-
 		FTransform WorldToObject = InObject->Transform.Inverse();
 
-		for (int32 Index = 0; Index < InObject->TrianglesArray.Num(); ++Index)
+		for (const IFC::FPolygon& IFCPolygon : InObject->TrianglesArray)
 		{
-			const IFC::FPolygon* IFCPolygon = &InObject->TrianglesArray[Index];
+			bool bInvalidTriangle = false;
 
-			TArray<FVertexInstanceID> VertexInstanceIDs;
 			TArray<FVertexID> VertexIDs;
+			TArray<FVertexInstanceID> VertexInstanceIDs;
 
-			// flip polygon to fix its orientation
-			for (int32 P = IFCPolygon->Points.Num() - 1; P > -1; --P)
+			// Flip polygon to fix its orientation
+			for (int32 IndexInPolygon = IFCPolygon.Points.Num() - 1; IndexInPolygon > -1; --IndexInPolygon)
 			{
-				const int32 VertexIndex = IFCPolygon->Points[P];
-				const FVertexID VertexID = PositionIndexToVertexId[VertexIndex];
-				const FVertexInstanceID& VertexInstanceID = MeshDescription->CreateVertexInstance(VertexID);
+				const int32 IFCVertexIndex = IFCPolygon.Points[IndexInPolygon];
+				const FVertexID& VertexID = PositionIndexToVertexId[IFCVertexIndex];
 
-				for (int32 UVIndex = 0; UVIndex < NumUVs; ++UVIndex)
+				if (VertexIDs.Num() > 0 && VertexID == VertexIDs.Last())
 				{
-					VertexInstanceUVs.Set(VertexInstanceID, UVIndex, FVector2D::ZeroVector);
+					bInvalidTriangle = true;
+					break;
 				}
 
-				VertexInstanceIDs.Add(VertexInstanceID);
 				VertexIDs.Add(VertexID);
+			}
 
-				const float* Vertex = &(InObject->facesVertices[(VertexIndex * (InObject->vertexElementSize / sizeof(float)))]);
-
-				// Flip Y to go from RH -> LH
-				FVector Normal = WorldToObject.TransformVector(FVector(Vertex[3], -Vertex[4], Vertex[5]));
-				VertexInstanceNormals.Set(VertexInstanceID, Normal.GetSafeNormal());
+			if (bInvalidTriangle)
+			{
+				continue;
 			}
 
-			bool bIsWrong = false;
-			for (int32 T = 1; T < VertexIDs.Num(); T++)
+			for(int32 IndexInPolygon = 0; IndexInPolygon < IFCPolygon.Points.Num(); ++IndexInPolygon)
 			{
-				if (VertexIDs[T - 1] == VertexIDs[T])
+				const int32 IFCVertexIndex = IFCPolygon.Points[IndexInPolygon];
+				const FVertexID& VertexID = VertexIDs[IndexInPolygon];
+
+				const FVertexInstanceID& VertexInstanceID = MeshDescription->CreateVertexInstance(VertexID);
+				VertexInstanceIDs.Add(VertexInstanceID);
+
+				for (int32 UVIndex = 0; UVIndex < NumUVs; ++UVIndex)
 				{
-					bIsWrong = true;
+					VertexInstanceUVs.Set(VertexInstanceID, UVIndex, FVector2D::ZeroVector);
 				}
-			}
-			if (bIsWrong)
-			{
-				continue;	// Skip this data.
+
+				const float* Vertex = &(InObject->facesVertices[(IFCVertexIndex * (InObject->vertexElementSize / sizeof(float)))]);
+
+				// Flip Y to go from RH -> LH
+				FVector Normal = WorldToObject.TransformVector(FVector(Vertex[3], -Vertex[4], Vertex[5]));
+				VertexInstanceNormals.Set(VertexInstanceID, Normal.GetSafeNormal());
 			}
 
-			MeshDescription->CreatePolygon(MaterialIndexToPolygonGroupID[IFCPolygon->MaterialIndex], VertexInstanceIDs);
+			MeshDescription->CreatePolygon(MaterialIndexToPolygonGroupID[IFCPolygon.MaterialIndex], VertexInstanceIDs);
 		}
 
 		DatasmithMeshHelper::RemoveEmptyPolygonGroups(*MeshDescription);
diff --git a/Engine/Plugins/Enterprise/DatasmithIFCImporter/Source/DatasmithIFCTranslator/Public/DatasmithIFCTranslatorModule.h b/Engine/Plugins/Enterprise/DatasmithIFCImporter/Source/DatasmithIFCTranslator/Public/DatasmithIFCTranslatorModule.h
--- a/Engine/Plugins/Enterprise/DatasmithIFCImporter/Source/DatasmithIFCTranslator/Public/DatasmithIFCTranslatorModule.h
+++ b/Engine/Plugins/Enterprise/DatasmithIFCImporter/Source/DatasmithIFCTranslator/Public/DatasmithIFCTranslatorModule.h
@@ -5,6 +5,8 @@
 #include "CoreMinimal.h"
 #include "Modules/ModuleManager.h"
 
+DECLARE_LOG_CATEGORY_EXTERN(LogIFCTranslator, Log, All);
+
 class IDatasmithIFCTranslatorModule : public IModuleInterface
 {
 public:
diff --git a/Engine/Plugins/Enterprise/DatasmithImporter/Source/DatasmithImporter/Private/DatasmithImportContext.cpp b/Engine/Plugins/Enterprise/DatasmithImporter/Source/DatasmithImporter/Private/DatasmithImportContext.cpp
--- a/Engine/Plugins/Enterprise/DatasmithImporter/Source/DatasmithImporter/Private/DatasmithImportContext.cpp
+++ b/Engine/Plugins/Enterprise/DatasmithImporter/Source/DatasmithImporter/Private/DatasmithImportContext.cpp
@@ -136,7 +136,7 @@ FDatasmithImportContext::FDatasmithImportContext(const FString& FileName, bool b
 	, bUserCancelled(false)
 	, bIsAReimport(false)
 	, bImportedViaScript(false)
-	, Warn(nullptr)
+	, FeedbackContext(nullptr)
 	, AssetsContext(*this)
 	, ContextExtension(nullptr)
 	, Logger(LoggerName, LoggerLabel)
@@ -303,7 +303,7 @@ bool FDatasmithImportContext::Init(TSharedRef< IDatasmithScene > InScene, const
 		}
 	}
 
-	Warn = InWarn;
+	FeedbackContext = InWarn;
 	Scene = InScene;
 
 	// Initialize the filtered scene as a copy of the original scene. We will use it to then filter out items to import.
diff --git a/Engine/Plugins/Enterprise/DatasmithImporter/Source/DatasmithImporter/Private/DatasmithImportFactory.cpp b/Engine/Plugins/Enterprise/DatasmithImporter/Source/DatasmithImporter/Private/DatasmithImportFactory.cpp
--- a/Engine/Plugins/Enterprise/DatasmithImporter/Source/DatasmithImporter/Private/DatasmithImportFactory.cpp
+++ b/Engine/Plugins/Enterprise/DatasmithImporter/Source/DatasmithImporter/Private/DatasmithImportFactory.cpp
@@ -29,8 +29,9 @@
 #include "Misc/PackageName.h"
 #include "Misc/ScopedSlowTask.h"
 #include "PackageTools.h"
-#include "Subsystems/ImportSubsystem.h"
 #include "Subsystems/AssetEditorSubsystem.h"
+#include "Subsystems/ImportSubsystem.h"
+#include "Templates/UniquePtr.h"
 
 #include "Editor/EditorEngine.h"
 
@@ -175,15 +176,24 @@ namespace DatasmithImportFactoryImpl
 			return false;
 		}
 
-		FScopedSlowTask Progress(100.0f, LOCTEXT("StartWork", "Unreal Datasmith ..."), true, *InContext.Warn);
-		Progress.MakeDialog(true);
+		TUniquePtr<FScopedSlowTask> ProgressPtr;
+		FScopedSlowTask* Progress = nullptr;
+		if ( InContext.FeedbackContext )
+		{
+			ProgressPtr = MakeUnique<FScopedSlowTask>(100.0f, LOCTEXT("StartWork", "Unreal Datasmith ..."), true, *InContext.FeedbackContext);
+			Progress = ProgressPtr.Get();
+			Progress->MakeDialog(true);
+		}
 
 		// Filter element that need to be imported depending on dirty state (or eventually depending on options)
 		FDatasmithImporter::FilterElementsToImport(InContext);
 
 		// TEXTURES
 		// We need the textures before the materials
-		Progress.EnterProgressFrame( 20.f );
+		if ( Progress )
+		{
+			Progress->EnterProgressFrame( 20.f );
+		}
 		FDatasmithImporter::ImportTextures(InContext);
 
 		if ( InContext.bUserCancelled )
@@ -194,7 +204,10 @@ namespace DatasmithImportFactoryImpl
 
 		// MATERIALS
 		// We need to import the materials before the static meshes to know about the meshes build requirements that are driven by the materials
-		Progress.EnterProgressFrame( 5.f );
+		if ( Progress )
+		{ 
+			Progress->EnterProgressFrame( 5.f );
+		}
 		FDatasmithImporter::ImportMaterials(InContext);
 
 		if ( InContext.bUserCancelled )
@@ -210,7 +223,10 @@ namespace DatasmithImportFactoryImpl
 		}
 
 		// STATIC MESHES
-		Progress.EnterProgressFrame( 25.f );
+		if ( Progress )
+		{
+			Progress->EnterProgressFrame( 25.f );
+		}
 		FDatasmithImporter::ImportStaticMeshes( InContext );
 
 		if ( InContext.bUserCancelled )
@@ -219,7 +235,10 @@ namespace DatasmithImportFactoryImpl
 			return false;
 		}
 
-		Progress.EnterProgressFrame( 10.f );
+		if ( Progress )
+		{
+			Progress->EnterProgressFrame( 10.f );
+		}
 		FDatasmithStaticMeshImporter::PreBuildStaticMeshes(InContext);
 
 		if ( InContext.bUserCancelled )
@@ -231,7 +250,10 @@ namespace DatasmithImportFactoryImpl
 		// ACTORS
 		if( InContext.ShouldImportActors() )
 		{
-			Progress.EnterProgressFrame( 10.f );
+			if ( Progress )
+			{
+				Progress->EnterProgressFrame( 10.f );
+			}
 
 			FDatasmithImporter::ImportActors( InContext );
 
@@ -254,7 +276,10 @@ namespace DatasmithImportFactoryImpl
 			return false;
 		}
 
-		Progress.EnterProgressFrame( 30.f );
+		if ( Progress )
+		{
+			Progress->EnterProgressFrame( 30.f );
+		}
 		FDatasmithImporter::FinalizeImport(InContext, TSet<UObject*>());
 
 		// THUMBNAIL
