VRED/DG: Merge of "Unmark aux file paths as configs so that they don't get loaded from ini file" from CL 11244472

diff --git a/Engine/Plugins/Enterprise/DatasmithFBXImporter/Source/DatasmithDeltaGenTranslator/Private/DatasmithDeltaGenImportOptions.h b/Engine/Plugins/Enterprise/DatasmithFBXImporter/Source/DatasmithDeltaGenTranslator/Private/DatasmithDeltaGenImportOptions.h
--- a/Engine/Plugins/Enterprise/DatasmithFBXImporter/Source/DatasmithDeltaGenTranslator/Private/DatasmithDeltaGenImportOptions.h
+++ b/Engine/Plugins/Enterprise/DatasmithFBXImporter/Source/DatasmithDeltaGenTranslator/Private/DatasmithDeltaGenImportOptions.h
@@ -37,19 +37,19 @@ public:
 	UPROPERTY(config, EditAnywhere, BlueprintReadWrite, Category=AssetImporting, meta=(DisplayName="Import Variants", ToolTip="import VAR files"))
 	bool bImportVar;
 
-	UPROPERTY(config, EditAnywhere, BlueprintReadWrite, Category=AssetImporting, meta=(DisplayName="Variants file path", EditCondition = "bImportVar", ToolTip="Path to the *.var file. By default it will search for a *.var file in the same folder as the FBX file, with the same base filename as it", FilePathFilter="var"))
+	UPROPERTY(EditAnywhere, BlueprintReadWrite, Category=AssetImporting, meta=(DisplayName="Variants file path", EditCondition = "bImportVar", ToolTip="Path to the *.var file. By default it will search for a *.var file in the same folder as the FBX file, with the same base filename as it", FilePathFilter="var"))
 	FFilePath VarPath;
 
 	UPROPERTY(config, EditAnywhere, BlueprintReadWrite, Category=AssetImporting, meta=(DisplayName="Import POS States", ToolTip="import POS files"))
 	bool bImportPos;
 
-	UPROPERTY(config, EditAnywhere, BlueprintReadWrite, Category=AssetImporting, meta=(DisplayName="POS file path", EditCondition = "bImportPos", ToolTip="Path to the *.pos file. By default it will search for a *.pos file in the same folder as the FBX file, with the same base filename as it", FilePathFilter="pos"))
+	UPROPERTY(EditAnywhere, BlueprintReadWrite, Category=AssetImporting, meta=(DisplayName="POS file path", EditCondition = "bImportPos", ToolTip="Path to the *.pos file. By default it will search for a *.pos file in the same folder as the FBX file, with the same base filename as it", FilePathFilter="pos"))
 	FFilePath PosPath;
 
 	UPROPERTY(config, EditAnywhere, BlueprintReadWrite, Category=AssetImporting, meta=(DisplayName="Import TML Animations", ToolTip="import TML files"))
 	bool bImportTml;
 
-	UPROPERTY(config, EditAnywhere, BlueprintReadWrite, Category=AssetImporting, meta=(DisplayName="TML file path", EditCondition = "bImportTml", ToolTip="Path to the *.tml file. By default it will search for a *.tml file in the same folder as the FBX file, with the same base filename as it", FilePathFilter="tml"))
+	UPROPERTY(EditAnywhere, BlueprintReadWrite, Category=AssetImporting, meta=(DisplayName="TML file path", EditCondition = "bImportTml", ToolTip="Path to the *.tml file. By default it will search for a *.tml file in the same folder as the FBX file, with the same base filename as it", FilePathFilter="tml"))
 	FFilePath TmlPath;
 
 	UPROPERTY(config, EditAnywhere, BlueprintReadWrite, Category=AssetImporting, meta=(DisplayName="Shadow Textures", ToolTip="How to handle shadow textures"))
diff --git a/Engine/Plugins/Enterprise/DatasmithFBXImporter/Source/DatasmithFBXTranslator/Public/DatasmithFBXImportOptions.h b/Engine/Plugins/Enterprise/DatasmithFBXImporter/Source/DatasmithFBXTranslator/Public/DatasmithFBXImportOptions.h
--- a/Engine/Plugins/Enterprise/DatasmithFBXImporter/Source/DatasmithFBXTranslator/Public/DatasmithFBXImportOptions.h
+++ b/Engine/Plugins/Enterprise/DatasmithFBXImporter/Source/DatasmithFBXTranslator/Public/DatasmithFBXImportOptions.h
@@ -15,6 +15,6 @@ class DATASMITHFBXTRANSLATOR_API UDatasmithFBXImportOptions : public UObject
 	GENERATED_UCLASS_BODY()
 
 public:
-	UPROPERTY(config, EditAnywhere, BlueprintReadWrite, Category=AssetImporting, meta=(DisplayName="Texture folders", ToolTip="Where to look for textures"))
+	UPROPERTY(EditAnywhere, BlueprintReadWrite, Category=AssetImporting, meta=(DisplayName="Texture folders", ToolTip="Where to look for textures"))
 	TArray<FDirectoryPath> TextureDirs;
 };
diff --git a/Engine/Plugins/Enterprise/DatasmithFBXImporter/Source/DatasmithVREDTranslator/Private/DatasmithVREDImportOptions.h b/Engine/Plugins/Enterprise/DatasmithFBXImporter/Source/DatasmithVREDTranslator/Private/DatasmithVREDImportOptions.h
--- a/Engine/Plugins/Enterprise/DatasmithFBXImporter/Source/DatasmithVREDTranslator/Private/DatasmithVREDImportOptions.h
+++ b/Engine/Plugins/Enterprise/DatasmithFBXImporter/Source/DatasmithVREDTranslator/Private/DatasmithVREDImportOptions.h
@@ -22,7 +22,7 @@ public:
 	UPROPERTY(config, EditAnywhere, BlueprintReadWrite, Category= AssetImporting, meta=(DisplayName="Import materials file", ToolTip="Uses the *.mats file saved alongside the exported FBX for a more accurate material reproduction"))
 	bool bImportMats;
 
-	UPROPERTY(config, EditAnywhere, BlueprintReadWrite, Category=AssetImporting, meta=(DisplayName="Materials file path", EditCondition = "bImportMats", ToolTip="Path to the *.mats file. By default it will search for a *.mats file in the same folder as the FBX file, with the same base filename as it", FilePathFilter="mats"))
+	UPROPERTY(EditAnywhere, BlueprintReadWrite, Category=AssetImporting, meta=(DisplayName="Materials file path", EditCondition = "bImportMats", ToolTip="Path to the *.mats file. By default it will search for a *.mats file in the same folder as the FBX file, with the same base filename as it", FilePathFilter="mats"))
 	FFilePath MatsPath;
 
 	UPROPERTY(config, EditAnywhere, BlueprintReadWrite, Category= AssetImporting, meta=(DisplayName="Import variants file", ToolTip="Uses the *.var file saved alongside the exported FBX"))
@@ -31,19 +31,19 @@ public:
 	UPROPERTY(config, EditAnywhere, AdvancedDisplay, BlueprintReadWrite, Category=AssetImporting, meta=(DisplayName="Clean variants", ToolTip="Removes empty variants, variant sets and invalid options. All discarded items will be logged to console."))
 	bool bCleanVar;
 
-	UPROPERTY(config, EditAnywhere, BlueprintReadWrite, Category=AssetImporting, meta=(DisplayName="Variants file path", EditCondition = "bImportVar", ToolTip="Path to the *.var file. By default it will search for a *.var file in the same folder as the FBX file, with the same base filename as it", FilePathFilter="var"))
+	UPROPERTY(EditAnywhere, BlueprintReadWrite, Category=AssetImporting, meta=(DisplayName="Variants file path", EditCondition = "bImportVar", ToolTip="Path to the *.var file. By default it will search for a *.var file in the same folder as the FBX file, with the same base filename as it", FilePathFilter="var"))
 	FFilePath VarPath;
 
 	UPROPERTY(config, EditAnywhere, BlueprintReadWrite, Category=AssetImporting, meta=(DisplayName="Import lights file", ToolTip="Uses the *.lights file saved alongside the exported FBX to import extra information about lights not saved in the FBX file"))
 	bool bImportLightInfo;
 
-	UPROPERTY(config, EditAnywhere, BlueprintReadWrite, Category=AssetImporting, meta=(DisplayName="Lights file path", EditCondition = "bImportLightInfo", ToolTip="Path to the *.lights file. By default it will search for a *.light file in the same folder as the FBX file, with the same base filename as it", FilePathFilter="lights"))
+	UPROPERTY(EditAnywhere, BlueprintReadWrite, Category=AssetImporting, meta=(DisplayName="Lights file path", EditCondition = "bImportLightInfo", ToolTip="Path to the *.lights file. By default it will search for a *.light file in the same folder as the FBX file, with the same base filename as it", FilePathFilter="lights"))
 	FFilePath LightInfoPath;
 
 	UPROPERTY(config, EditAnywhere, BlueprintReadWrite, Category=AssetImporting, meta=(DisplayName="Import clips file", ToolTip="Uses the *.clips file saved alongside the exported FBX to import information about animation clips and blocks, mirroring VRED's animation system"))
 	bool bImportClipInfo;
 
-	UPROPERTY(config, EditAnywhere, BlueprintReadWrite, Category=AssetImporting, meta=(DisplayName="AnimClips file path", EditCondition = "bImportClipInfo", ToolTip="Path to the *.clips file. By default it will search for a *.clips file in the same folder as the FBX file, with the same base filename as it", FilePathFilter="clips"))
+	UPROPERTY(EditAnywhere, BlueprintReadWrite, Category=AssetImporting, meta=(DisplayName="AnimClips file path", EditCondition = "bImportClipInfo", ToolTip="Path to the *.clips file. By default it will search for a *.clips file in the same folder as the FBX file, with the same base filename as it", FilePathFilter="clips"))
 	FFilePath ClipInfoPath;
 
 public:
