DG: Combine multiple transform animations into one

Integrate from Dev-Enterprise CL 11215719

diff --git a/Engine/Plugins/Enterprise/DatasmithFBXImporter/Source/DatasmithDeltaGenTranslator/Private/DatasmithDeltaGenImporter.cpp b/Engine/Plugins/Enterprise/DatasmithFBXImporter/Source/DatasmithDeltaGenTranslator/Private/DatasmithDeltaGenImporter.cpp
--- a/Engine/Plugins/Enterprise/DatasmithFBXImporter/Source/DatasmithDeltaGenTranslator/Private/DatasmithDeltaGenImporter.cpp
+++ b/Engine/Plugins/Enterprise/DatasmithFBXImporter/Source/DatasmithDeltaGenTranslator/Private/DatasmithDeltaGenImporter.cpp
@@ -898,12 +898,50 @@ namespace DeltaGenImporterImpl
 		ETransformChannelComponents Components = ETransformChannelComponents::All;
 		TransformAnimation.SetEnabledTransformChannels(Channels | FDatasmithAnimationUtils::SetChannelTypeComponents(Components, DSType));
 	}
+
+	void CombineAnimations(const TArray<TSharedRef<IDatasmithTransformAnimationElement>>& InputAnimations, TSharedRef<IDatasmithTransformAnimationElement>& OutputAnimation)
+	{
+		TArray<FDatasmithTransformFrameInfo> Frames;
+		for (uint8 TransformTypeIndex = 0; TransformTypeIndex < (uint8)EDatasmithTransformType::Count; ++TransformTypeIndex)
+		{
+			EDatasmithTransformType TransformType = (EDatasmithTransformType)TransformTypeIndex;
+
+			int32 NumFrames = 0;
+			for (const TSharedRef<IDatasmithTransformAnimationElement>& Animation : InputAnimations)
+			{
+				NumFrames += Animation->GetFramesCount(TransformType);
+			}
+			Frames.Reset();
+			Frames.Reserve(NumFrames);
+
+			for (const TSharedRef<IDatasmithTransformAnimationElement>& Animation : InputAnimations)
+			{
+				int32 FrameCount = Animation->GetFramesCount(TransformType);
+				for (int32 FrameIndex = 0; FrameIndex < FrameCount; ++FrameIndex)
+				{
+					Frames.Add(Animation->GetFrame(TransformType, FrameIndex));
+				}
+			}
+
+			Algo::Sort(Frames, [](const FDatasmithTransformFrameInfo& A, const FDatasmithTransformFrameInfo& B)
+			{
+				return A.FrameNumber < B.FrameNumber;
+			});
+
+			for (const FDatasmithTransformFrameInfo& Frame : Frames)
+			{
+				OutputAnimation->AddFrame(TransformType, Frame);
+			}
+		}
+	}
 }
 
 TSharedPtr<IDatasmithLevelSequenceElement> FDatasmithDeltaGenImporter::ConvertAnimationTimeline(const FDeltaGenTmlDataTimeline& TmlTimeline)
 {
 	TSharedRef<IDatasmithLevelSequenceElement> SequenceElement = FDatasmithSceneFactory::CreateLevelSequence(*TmlTimeline.Name);
+	SequenceElement->SetFrameRate(TmlTimeline.Framerate);
 
+	TArray<TSharedRef<IDatasmithTransformAnimationElement>> InitialAnimations;
 	for (const FDeltaGenTmlDataTimelineAnimation& Animation : TmlTimeline.Animations)
 	{
 		FString TargetNodeName = Animation.TargetNode.ToString();
@@ -921,8 +959,37 @@ TSharedPtr<IDatasmithLevelSequenceElement> FDatasmithDeltaGenImporter::ConvertAn
 			TransformAnimation->GetFramesCount(EDatasmithTransformType::Rotation) > 0 ||
 			TransformAnimation->GetFramesCount(EDatasmithTransformType::Scale) > 0)
 		{
-			SequenceElement->AddAnimation(TransformAnimation);
-			SequenceElement->SetFrameRate(TmlTimeline.Framerate);
+			InitialAnimations.Add(TransformAnimation);
+		}
+	}
+
+	// Temp fix for UE-87458 until DatasmithLevelSequenceImporter can handle multiple simultaneous animations per actor using separate sections.
+	// This combines animations so that we will only ever emit one IDatasmithTransformAnimationElement per actor per level sequence.
+	// This needs to be done after baking as tracks can have different interpolation modes.
+	{
+		TMap<FString, TArray<TSharedRef<IDatasmithTransformAnimationElement>>> AnimationsPerActor;
+
+		for (const TSharedRef<IDatasmithTransformAnimationElement>& Animation : InitialAnimations)
+		{
+			AnimationsPerActor.FindOrAdd(Animation->GetName()).Add(Animation);
+		}
+
+		for (const TPair<FString, TArray<TSharedRef<IDatasmithTransformAnimationElement>>>& Pair : AnimationsPerActor)
+		{
+			const FString& ActorName = Pair.Key;
+			const TArray<TSharedRef<IDatasmithTransformAnimationElement>>& Animations = Pair.Value;
+
+			if (Animations.Num() == 1)
+			{
+				SequenceElement->AddAnimation(Animations[0]);
+			}
+			else
+			{
+				TSharedRef<IDatasmithTransformAnimationElement> CombinedAnimation = FDatasmithSceneFactory::CreateTransformAnimation(*ActorName);
+				DeltaGenImporterImpl::CombineAnimations(Animations, CombinedAnimation);
+
+				SequenceElement->AddAnimation(CombinedAnimation);
+			}
 		}
 	}
 
