DeltaGen: Enable support for shininess textures

Integrate from Dev-Enterprise cl 11383119

diff --git a/Engine/Plugins/Enterprise/DatasmithFBXImporter/Source/DatasmithDeltaGenTranslator/Private/DatasmithDeltaGenImporter.cpp b/Engine/Plugins/Enterprise/DatasmithFBXImporter/Source/DatasmithDeltaGenTranslator/Private/DatasmithDeltaGenImporter.cpp
--- a/Engine/Plugins/Enterprise/DatasmithFBXImporter/Source/DatasmithDeltaGenTranslator/Private/DatasmithDeltaGenImporter.cpp
+++ b/Engine/Plugins/Enterprise/DatasmithFBXImporter/Source/DatasmithDeltaGenTranslator/Private/DatasmithDeltaGenImporter.cpp
@@ -579,6 +579,7 @@ TSharedPtr<IDatasmithTextureElement> CreateTextureAndTextureProperties(IDatasmit
 		MapEntry(TEXT("TexTransparent"), EDatasmithTextureMode::Other),
 		MapEntry(TEXT("TexEmissive"), EDatasmithTextureMode::Diffuse),
 		MapEntry(TEXT("TexAO"), EDatasmithTextureMode::Diffuse),
+		MapEntry(TEXT("TexShininess"), EDatasmithTextureMode::Specular),
 	});
 
 	// Create the actual texture (accompanying texture properties will all be packed as key-value pairs)
diff --git a/Engine/Plugins/Enterprise/DatasmithFBXImporter/Source/DatasmithFBXTranslator/Private/DatasmithFBXFileImporter.cpp b/Engine/Plugins/Enterprise/DatasmithFBXImporter/Source/DatasmithFBXTranslator/Private/DatasmithFBXFileImporter.cpp
--- a/Engine/Plugins/Enterprise/DatasmithFBXImporter/Source/DatasmithFBXTranslator/Private/DatasmithFBXFileImporter.cpp
+++ b/Engine/Plugins/Enterprise/DatasmithFBXImporter/Source/DatasmithFBXTranslator/Private/DatasmithFBXFileImporter.cpp
@@ -1381,6 +1381,7 @@ TSharedPtr<FDatasmithFBXSceneMaterial> FDatasmithFBXFileImporter::ImportMaterial
 			FetchFbxTexture(InMaterial, TEXT("ReflectionColor"), *Material, PathStr);
 			FetchFbxTexture(InMaterial, TEXT("TransparentColor"), *Material, PathStr);
 			FetchFbxTexture(InMaterial, TEXT("EmissiveColor"), *Material, PathStr);
+			FetchFbxTexture(InMaterial, TEXT("Shininess"), *Material, PathStr);
 		}
 	}
 
