Avoid parsing the whole FBX scene for each translator compatibility detection
  - 40s to 10s to reach import dialog after choosing DATASET-0022a
  - 54s to 3s to reach import dialog after choosing DATASET-0023a

From dev-enterprise[REDACTED_HANDLE]

diff --git a/Engine/Plugins/Enterprise/DatasmithFBXImporter/Source/DatasmithDeltaGenTranslator/Private/DatasmithDeltaGenTranslator.cpp b/Engine/Plugins/Enterprise/DatasmithFBXImporter/Source/DatasmithDeltaGenTranslator/Private/DatasmithDeltaGenTranslator.cpp
--- a/Engine/Plugins/Enterprise/DatasmithFBXImporter/Source/DatasmithDeltaGenTranslator/Private/DatasmithDeltaGenTranslator.cpp
+++ b/Engine/Plugins/Enterprise/DatasmithFBXImporter/Source/DatasmithDeltaGenTranslator/Private/DatasmithDeltaGenTranslator.cpp
@@ -7,7 +7,7 @@
 #include "IDatasmithSceneElements.h"
 #include "DatasmithDeltaGenImporter.h"
 #include "DatasmithDeltaGenLog.h"
-
+#include "HAL/FileManager.h"
 #include "FbxImporter.h"
 #include "MeshDescription.h"
 
@@ -22,6 +22,8 @@ void FDatasmithDeltaGenTranslator::Initialize(FDatasmithTranslatorCapabilities&
 
 bool FDatasmithDeltaGenTranslator::IsSourceSupported(const FDatasmithSceneSource& Source)
 {
+	TRACE_CPUPROFILER_EVENT_SCOPE(FDatasmithDeltaGenTranslator::IsSourceSupported)
+
 	const FString& FilePath = Source.GetSourceFile();
 	const FString& Extension = Source.GetSourceFileExtension();
 	if (!Extension.Equals(TEXT("fbx"), ESearchCase::IgnoreCase))
@@ -29,29 +31,54 @@ bool FDatasmithDeltaGenTranslator::IsSourceSupported(const FDatasmithSceneSource
 		return false;
 	}
 
-	UnFbx::FFbxImporter* FbxImporter = UnFbx::FFbxImporter::GetInstance();
-	UnFbx::FBXImportOptions* GlobalImportSettings = FbxImporter->GetImportOptions();
-	UnFbx::FBXImportOptions::ResetOptions(GlobalImportSettings);
-
-	if (!FbxImporter->ImportFromFile(FilePath, Extension, false))
+	FArchive* Reader = IFileManager::Get().CreateFileReader( *FilePath );
+	if( !Reader )
 	{
-		FbxImporter->ReleaseScene();
 		return false;
 	}
 
-	FString ProductName = UTF8_TO_TCHAR(FbxImporter->Scene->GetSceneInfo()->Original_ApplicationName.Get().Buffer());
-	if (ProductName != TEXT("RTT_AG"))
+	ANSICHAR Header[64*1024] = { 0 };
+	Reader->Serialize(Header, FMath::Min(Reader->TotalSize(), (int64)sizeof(Header) - 1));
+	delete Reader;
+
+	// Replace 0 with anything for Strstr to work on binary files
+	for (int32 Index = 0; Index < sizeof(Header) - 1; ++Index)
 	{
-		FbxImporter->ReleaseScene();
-		return false;
+		if (Header[Index] == '\0')
+		{
+			Header[Index] = '.';
+		}
 	}
 
-	FbxImporter->ReleaseScene();
-	return true;
+	// Quick and dirty way of identifying with a high degree of confidence if the file
+	// is from DeltaGen without parsing the whole scene using the SDK which can take a
+	// long time for big scenes.
+	//
+	// Supports both ASCII and binary formats.
+	{
+		const ANSICHAR* TagName = FPlatformString::Strstr(Header, "Original|ApplicationName");
+		if (TagName)
+		{
+			const ANSICHAR* TagType = FPlatformString::Strstr(TagName, "KString");
+			if (TagType)
+			{
+				const ANSICHAR* TagData = FPlatformString::Strstr(TagType, "RTT_AG");
+				if (TagData)
+				{
+					// The whole tag should be in the same vicinity
+					return (TagData - TagName) < 256;
+				}
+			}
+		}
+	}
+
+	return false;
 }
 
 bool FDatasmithDeltaGenTranslator::LoadScene(TSharedRef<IDatasmithScene> OutScene)
 {
+	TRACE_CPUPROFILER_EVENT_SCOPE(FDatasmithDeltaGenTranslator::LoadScene)
+
 	OutScene->SetHost(TEXT("DeltaGenTranslator"));
 
     Importer = MakeShared<FDatasmithDeltaGenImporter>(OutScene, ImportOptions.Get());
diff --git a/Engine/Plugins/Enterprise/DatasmithFBXImporter/Source/DatasmithVREDTranslator/Private/DatasmithVREDTranslator.cpp b/Engine/Plugins/Enterprise/DatasmithFBXImporter/Source/DatasmithVREDTranslator/Private/DatasmithVREDTranslator.cpp
--- a/Engine/Plugins/Enterprise/DatasmithFBXImporter/Source/DatasmithVREDTranslator/Private/DatasmithVREDTranslator.cpp
+++ b/Engine/Plugins/Enterprise/DatasmithFBXImporter/Source/DatasmithVREDTranslator/Private/DatasmithVREDTranslator.cpp
@@ -7,7 +7,7 @@
 #include "DatasmithVREDLog.h"
 #include "DatasmithVREDTranslatorModule.h"
 #include "IDatasmithSceneElements.h"
-
+#include "HAL/FileManager.h"
 #include "FbxImporter.h"
 #include "MeshDescription.h"
 
@@ -22,6 +22,8 @@ void FDatasmithVREDTranslator::Initialize(FDatasmithTranslatorCapabilities& OutC
 
 bool FDatasmithVREDTranslator::IsSourceSupported(const FDatasmithSceneSource& Source)
 {
+	TRACE_CPUPROFILER_EVENT_SCOPE(FDatasmithVREDTranslator::IsSourceSupported)
+
 	const FString& FilePath = Source.GetSourceFile();
 	const FString& Extension = Source.GetSourceFileExtension();
 	if (!Extension.Equals(TEXT("fbx"), ESearchCase::IgnoreCase))
@@ -29,30 +31,73 @@ bool FDatasmithVREDTranslator::IsSourceSupported(const FDatasmithSceneSource& So
 		return false;
 	}
 
-	UnFbx::FFbxImporter* FbxImporter = UnFbx::FFbxImporter::GetInstance();
-	UnFbx::FBXImportOptions* GlobalImportSettings = FbxImporter->GetImportOptions();
-	UnFbx::FBXImportOptions::ResetOptions(GlobalImportSettings);
-
-	if (!FbxImporter->ImportFromFile(FilePath, Extension, false))
+	FArchive* Reader = IFileManager::Get().CreateFileReader( *FilePath );
+	if( !Reader )
 	{
-		FbxImporter->ReleaseScene();
 		return false;
 	}
 
-	FString ProductName = UTF8_TO_TCHAR(FbxImporter->Scene->GetSceneInfo()->Original_ApplicationName.Get().Buffer());
-	FString ProductVendor = UTF8_TO_TCHAR(FbxImporter->Scene->GetSceneInfo()->Original_ApplicationVendor.Get().Buffer());
-	if (ProductName != TEXT("VRED") || ProductVendor != TEXT("Autodesk"))
+	ANSICHAR Header[64*1024] = { 0 };
+	Reader->Serialize(Header, FMath::Min(Reader->TotalSize(), (int64)sizeof(Header) - 1));
+	delete Reader;
+
+	// Replace 0 with anything for Strstr to work on binary files
+	for (int32 Index = 0; Index < sizeof(Header) - 1; ++Index)
 	{
-		FbxImporter->ReleaseScene();
-		return false;
+		if (Header[Index] == '\0')
+		{
+			Header[Index] = '.';
+		}
 	}
 
-	FbxImporter->ReleaseScene();
-	return true;
+	// Quick and dirty way of identifying with a high degree of confidence if the file
+	// is from VRED without parsing the whole scene using the SDK which can take a
+	// long time for big scenes.
+	//
+	// Supports both ASCII and binary formats.
+	bool bApplicationVendorValid = false;
+	{
+		const ANSICHAR* TagName = FPlatformString::Strstr(Header, "Original|ApplicationVendor");
+		if (TagName)
+		{
+			const ANSICHAR* TagType = FPlatformString::Strstr(TagName, "KString");
+			if (TagType)
+			{
+				const ANSICHAR* TagData = FPlatformString::Strstr(TagType, "Autodesk");
+				if (TagData)
+				{
+					// The whole tag should be in the same vicinity
+					bApplicationVendorValid = (TagData - TagName) < 256;
+				}
+			}
+		}
+	}
+
+	bool bApplicationNameValid = false;
+	{
+		const ANSICHAR* TagName = FPlatformString::Strstr(Header, "Original|ApplicationName");
+		if (TagName)
+		{
+			const ANSICHAR* TagType = FPlatformString::Strstr(TagName, "KString");
+			if (TagType)
+			{
+				const ANSICHAR* TagData = FPlatformString::Strstr(TagType, "VRED");
+				if (TagData)
+				{
+					// The whole tag should be in the same vicinity
+					bApplicationNameValid = (TagData - TagName) < 256;
+				}
+			}
+		}
+	}
+
+	return bApplicationVendorValid && bApplicationNameValid;
 }
 
 bool FDatasmithVREDTranslator::LoadScene(TSharedRef<IDatasmithScene> OutScene)
 {
+	TRACE_CPUPROFILER_EVENT_SCOPE(FDatasmithVREDTranslator::LoadScene)
+
 	OutScene->SetHost(TEXT("VREDTranslator"));
 	OutScene->SetProductName(TEXT("VRED"));
 
