DeltaGen: Clone material instances if they are using diferent AO textures

diff --git a/Engine/Plugins/Enterprise/DatasmithFBXImporter/Source/DatasmithDeltaGenTranslator/Private/DatasmithDeltaGenImporter.cpp b/Engine/Plugins/Enterprise/DatasmithFBXImporter/Source/DatasmithDeltaGenTranslator/Private/DatasmithDeltaGenImporter.cpp
--- a/Engine/Plugins/Enterprise/DatasmithFBXImporter/Source/DatasmithDeltaGenTranslator/Private/DatasmithDeltaGenImporter.cpp
+++ b/Engine/Plugins/Enterprise/DatasmithFBXImporter/Source/DatasmithDeltaGenTranslator/Private/DatasmithDeltaGenImporter.cpp
@@ -207,65 +207,15 @@ void FDatasmithDeltaGenImporter::ParseAuxFiles(const FString& FBXPath)
 	}
 }
 
-void FDatasmithDeltaGenImporter::FetchAOTexture(const FString& MeshName, TSharedPtr<FDatasmithFBXSceneMaterial>& Material)
-{
-	if (ImportOptions->TextureDirs.Num() == 0 || ImportOptions->ShadowTextureMode == EShadowTextureMode::Ignore)
-	{
-		return;
-	}
-
-	// FNames because they're case insensitive
-	// I've only ever seen .bmp shadow textures, but I haven't seen anything saying they can't be anything else
-	const static TSet<FName> ImageExtensions{ TEXT("bmp"), TEXT("jpg"), TEXT("png"), TEXT("jpeg"), TEXT("tiff"), TEXT("tga") };
-
-	// Find all filepaths for images that are in a texture folder and have the mesh name as part of the filename
-	TArray<FString> PotentialTextures;
-	for (const FDirectoryPath& Dir : ImportOptions->TextureDirs)
-	{
-		TArray<FString> Textures;
-		IFileManager::Get().FindFiles(Textures, *Dir.Path, TEXT(""));
-
-		for (const FString& Texture : Textures)
-		{
-			FName Extension = FName(*FPaths::GetExtension(Texture));
-
-			if (ImageExtensions.Contains(Extension) && Texture.Contains(MeshName))
-			{
-				PotentialTextures.Add( Dir.Path / Texture );
-			}
-		}
-	}
-
-	if (PotentialTextures.Num() == 0)
-	{
-		return;
-	}
-
-	FString AOTexPath = PotentialTextures[0];
-	if (PotentialTextures.Num() > 1)
-	{
-		UE_LOG(LogDatasmithDeltaGenImport, Log, TEXT("Found more than one candidate for an AO texture for mesh '%s'. The texture '%s' will be used, but moving or renaming the texture would prevent this."), *MeshName, *AOTexPath);
-	}
-
-	FDatasmithFBXSceneMaterial::FTextureParams& Tex = Material->TextureParams.FindOrAdd(TEXT("TexAO"));
-	Tex.Path = AOTexPath;
-}
-
 bool FDatasmithDeltaGenImporter::ProcessScene()
 {
 	FDatasmithDeltaGenSceneProcessor Processor(IntermediateScene.Get());
 
-	// We need to create AO textures before we merge as they depend on the name of the mesh itself,
-	// and we only want to add this AO texture to the one material used by that mesh
+	Processor.FindDuplicatedMaterials();
+
 	if (ImportOptions->ShadowTextureMode != EShadowTextureMode::Ignore)
 	{
-		for (TSharedPtr<FDatasmithFBXSceneNode>& Node : IntermediateScene->GetAllNodes())
-		{
-			for (TSharedPtr<FDatasmithFBXSceneMaterial>& Material : Node->Materials)
-			{
-				FetchAOTexture(Node->Mesh->Name, Material);
-			}
-		}
+		Processor.SetupAOTextures(ImportOptions->TextureDirs);
 	}
 
 	Processor.RemoveLightMapNodes();
@@ -276,8 +226,6 @@ bool FDatasmithDeltaGenImporter::ProcessScene()
 
 	Processor.DecomposePivots(TmlTimelines);
 
-	Processor.FindDuplicatedMaterials();
-
 	if (ImportOptions->bRemoveInvisibleNodes)
 	{
 		Processor.RemoveInvisibleNodes();
diff --git a/Engine/Plugins/Enterprise/DatasmithFBXImporter/Source/DatasmithDeltaGenTranslator/Private/DatasmithDeltaGenImporter.h b/Engine/Plugins/Enterprise/DatasmithFBXImporter/Source/DatasmithDeltaGenTranslator/Private/DatasmithDeltaGenImporter.h
--- a/Engine/Plugins/Enterprise/DatasmithFBXImporter/Source/DatasmithDeltaGenTranslator/Private/DatasmithDeltaGenImporter.h
+++ b/Engine/Plugins/Enterprise/DatasmithFBXImporter/Source/DatasmithDeltaGenTranslator/Private/DatasmithDeltaGenImporter.h
@@ -52,8 +52,6 @@ private:
 
 	TSharedPtr<IDatasmithBaseMaterialElement> ConvertMaterial(const TSharedPtr<FDatasmithFBXSceneMaterial>& Material);
 
-	void FetchAOTexture(const FString& MeshName, TSharedPtr<FDatasmithFBXSceneMaterial>& Material);
-
 	TSharedPtr<IDatasmithLevelSequenceElement> ConvertAnimationTimeline(const FDeltaGenTmlDataTimeline& TmlTimeline);
 
 private:
diff --git a/Engine/Plugins/Enterprise/DatasmithFBXImporter/Source/DatasmithDeltaGenTranslator/Private/DatasmithDeltaGenSceneProcessor.cpp b/Engine/Plugins/Enterprise/DatasmithFBXImporter/Source/DatasmithDeltaGenTranslator/Private/DatasmithDeltaGenSceneProcessor.cpp
--- a/Engine/Plugins/Enterprise/DatasmithFBXImporter/Source/DatasmithDeltaGenTranslator/Private/DatasmithDeltaGenSceneProcessor.cpp
+++ b/Engine/Plugins/Enterprise/DatasmithFBXImporter/Source/DatasmithDeltaGenTranslator/Private/DatasmithDeltaGenSceneProcessor.cpp
@@ -3,12 +3,16 @@
 #include "DatasmithDeltaGenSceneProcessor.h"
 
 #include "DatasmithDeltaGenImportData.h"
+#include "DatasmithDeltaGenImportOptions.h"
 #include "DatasmithDeltaGenLog.h"
 #include "DatasmithFBXScene.h"
+#include "DatasmithUtils.h"
+
+#include "Factories\TextureFactory.h"
 
 using TimelineToAnimations = TMap<FDeltaGenTmlDataTimeline*, FDeltaGenTmlDataTimelineAnimation*>;
 
-namespace FDeltaGenProcessorImpl
+namespace DeltaGenProcessorImpl
 {
 	// Moves tracks of TrackTypes to new animations within each Timeline, all tied to the Dummy node's name
 	void MoveTracksToDummyAnimation(TSharedPtr<FDatasmithFBXSceneNode>& Dummy,
@@ -152,6 +156,99 @@ namespace FDeltaGenProcessorImpl
 		NodeParent->Children.Remove(Node);
 		NodeParent->AddChild(Dummy);
 	}
+
+	/** Will return a clone of Original that has TexAO set with 'AOPath' as its path. It may reuse an existing clone, or create a new one */
+	TSharedPtr<FDatasmithFBXSceneMaterial> CloneMaterialForAOTexture(const TSharedPtr<FDatasmithFBXSceneMaterial>& Original, const FString& AOPath)
+	{
+		if (!Original.IsValid() || AOPath.IsEmpty())
+		{
+			return nullptr;
+		}
+
+		// Maybe we already cloned this material for that AO texture?
+		for (const TWeakPtr<FDatasmithFBXSceneMaterial>& Clone : Original->ClonedMaterials)
+		{
+			if (!Clone.IsValid())
+			{
+				continue;
+			}
+
+			TSharedPtr<FDatasmithFBXSceneMaterial> PinnedClone = Clone.Pin();
+			if (FDatasmithFBXSceneMaterial::FTextureParams* AOTexture = PinnedClone->TextureParams.Find(TEXT("TexAO")))
+			{
+				if (AOTexture->Path == AOPath)
+				{
+					return PinnedClone;
+				}
+			}
+		}
+
+		TSharedPtr<FDatasmithFBXSceneMaterial> Clone = MakeShared<FDatasmithFBXSceneMaterial>(*Original);
+		Clone->TextureParams.FindOrAdd(TEXT("TexAO")).Path = AOPath;
+
+		Original->ClonedMaterials.Add(Clone);
+
+		return Clone;
+	}
+
+	/**
+	 * Will look for an existing AO texture for a given mesh name and set it on the material.
+	 * If it can't set the texture (already has a different AO texture), it will return a cloned version of the material with the new AO texture.
+	 */
+	TSharedPtr<FDatasmithFBXSceneMaterial> FetchAOTexture(const FString& MeshName, TSharedPtr<FDatasmithFBXSceneMaterial>& Material, const TArray<FDirectoryPath>& TextureFolders)
+	{
+		UTextureFactory* TextureFactory = UTextureFactory::StaticClass()->GetDefaultObject<UTextureFactory>();
+		if (TextureFolders.Num() == 0 || !TextureFactory)
+		{
+			return nullptr;
+		}
+
+		// Find all filepaths for images that are in a texture folder and have the mesh name as part of the filename
+		TArray<FString> PotentialTextures;
+		for (const FDirectoryPath& Dir : TextureFolders)
+		{
+			TArray<FString> Textures;
+			IFileManager::Get().FindFiles(Textures, *Dir.Path, TEXT(""));
+
+			for (const FString& Texture : Textures)
+			{
+				if (TextureFactory->FactoryCanImport(Texture) && Texture.Contains(MeshName))
+				{
+					PotentialTextures.Add(Dir.Path / Texture);
+				}
+			}
+		}
+
+		if (PotentialTextures.Num() == 0)
+		{
+			return nullptr;
+		}
+
+		FString AOTexPath = PotentialTextures[0];
+		if (PotentialTextures.Num() > 1)
+		{
+			UE_LOG(LogDatasmithDeltaGenImport, Warning, TEXT("Found more than one candidate for an AO texture matching mesh name '%s', so texture '%s' will be used. Better matching between texture and mesh name would help."), *MeshName, *AOTexPath);
+		}
+
+		TSharedPtr<FDatasmithFBXSceneMaterial> ClonedMaterial = nullptr;
+		if (FDatasmithFBXSceneMaterial::FTextureParams* ExistingAOTexture = Material->TextureParams.Find(TEXT("TexAO")))
+		{
+			// If the material already has a different AO texture set, some other mesh is using it, so we need
+			// to duplicate this material and set our own AO texture
+			if (ExistingAOTexture->Path != AOTexPath)
+			{
+				ClonedMaterial = DeltaGenProcessorImpl::CloneMaterialForAOTexture(Material, AOTexPath);
+			}
+		}
+		// No AO texture on this material yet, lets set our texture to it and use it
+		else
+		{
+			FDatasmithFBXSceneMaterial::FTextureParams& Tex = Material->TextureParams.FindOrAdd(TEXT("TexAO"));
+			Tex.Path = AOTexPath;
+		}
+
+		return ClonedMaterial;
+	}
 };
 
 FDatasmithDeltaGenSceneProcessor::FDatasmithDeltaGenSceneProcessor(FDatasmithFBXScene* InScene)
@@ -159,6 +256,36 @@ FDatasmithDeltaGenSceneProcessor::FDatasmithDeltaGenSceneProcessor(FDatasmithFBX
 {
 }
 
+void FDatasmithDeltaGenSceneProcessor::SetupAOTextures(const TArray<FDirectoryPath>& TextureFolders)
+{
+	TSet<TSharedPtr<FDatasmithFBXSceneMaterial>> UniqueMaterials{ Scene->Materials };
+
+	FDatasmithUniqueNameProvider UniqueNameProvider;
+	for (const TSharedPtr<FDatasmithFBXSceneMaterial>& UniqueMaterial : UniqueMaterials)
+	{
+		UniqueNameProvider.AddExistingName(UniqueMaterial->Name);
+	}
+
+	for (TSharedPtr<FDatasmithFBXSceneNode>& Node : Scene->GetAllNodes())
+	{
+		for (int32 Index = 0; Index < Node->Materials.Num(); ++Index)
+		{
+			TSharedPtr<FDatasmithFBXSceneMaterial>& Material = Node->Materials[Index];
+
+			TSharedPtr<FDatasmithFBXSceneMaterial> ClonedMaterial = DeltaGenProcessorImpl::FetchAOTexture(Node->Mesh->Name, Material, TextureFolders);
+			if (ClonedMaterial.IsValid())
+			{
+				ClonedMaterial->Name = UniqueNameProvider.GenerateUniqueName(ClonedMaterial->Name);
+
+				UniqueMaterials.Add(ClonedMaterial);
+				Node->Materials[Index] = ClonedMaterial;
+			}
+		}
+	}
+
+	Scene->Materials = UniqueMaterials.Array();
+}
+
 void FDatasmithDeltaGenSceneProcessor::DecomposePivots(TArray<FDeltaGenTmlDataTimeline>& Timelines)
 {
 	// Cache node names to all the animations they have on all timelines
@@ -176,8 +303,8 @@ void FDatasmithDeltaGenSceneProcessor::DecomposePivots(TArray<FDeltaGenTmlDataTi
 	TMap<FDeltaGenTmlDataTimeline*, TArray<FDeltaGenTmlDataTimelineAnimation>> NewAnimationsPerTimeline;
 	for (TSharedPtr<FDatasmithFBXSceneNode> Node : Scene->GetAllNodes())
 	{
-		FDeltaGenProcessorImpl::DecomposeRotationPivotsForNode(Node, NodeNamesToAnimations, NewAnimationsPerTimeline);
-		FDeltaGenProcessorImpl::DecomposeScalingPivotsForNode(Node, NodeNamesToAnimations, NewAnimationsPerTimeline);
+		DeltaGenProcessorImpl::DecomposeRotationPivotsForNode(Node, NodeNamesToAnimations, NewAnimationsPerTimeline);
+		DeltaGenProcessorImpl::DecomposeScalingPivotsForNode(Node, NodeNamesToAnimations, NewAnimationsPerTimeline);
 	}
 
 	// Add the new animations only afterwards, as NodeNamesToAnimations stores raw pointers to animations within TArrays
diff --git a/Engine/Plugins/Enterprise/DatasmithFBXImporter/Source/DatasmithDeltaGenTranslator/Private/DatasmithDeltaGenSceneProcessor.h b/Engine/Plugins/Enterprise/DatasmithFBXImporter/Source/DatasmithDeltaGenTranslator/Private/DatasmithDeltaGenSceneProcessor.h
--- a/Engine/Plugins/Enterprise/DatasmithFBXImporter/Source/DatasmithDeltaGenTranslator/Private/DatasmithDeltaGenSceneProcessor.h
+++ b/Engine/Plugins/Enterprise/DatasmithFBXImporter/Source/DatasmithDeltaGenTranslator/Private/DatasmithDeltaGenSceneProcessor.h
@@ -5,6 +5,7 @@
 #include "DatasmithFBXSceneProcessor.h"
 
 #include "CoreTypes.h"
+#include "Engine\EngineTypes.h"
 
 struct FDatasmithFBXScene;
 struct FDeltaGenTmlDataTimeline;
@@ -14,6 +15,12 @@ class FDatasmithDeltaGenSceneProcessor : public FDatasmithFBXSceneProcessor
 public:
 	FDatasmithDeltaGenSceneProcessor(FDatasmithFBXScene* InScene);
 
+	/**
+	 * Fetches AO textures matching corresponding mesh names, and assigns them to materials used for each node.
+	 * May create additional materials, as we may have to clone material instances to use different AO textures per mesh.
+	 */
+	void SetupAOTextures(const TArray<FDirectoryPath>& TextureFolders);
+
 	/**
 	 * Decompose all scene nodes with nonzero rotation and scaling pivots
 	 * using dummy actors, and handle their animations
diff --git a/Engine/Plugins/Enterprise/DatasmithFBXImporter/Source/DatasmithFBXTranslator/Public/DatasmithFBXScene.h b/Engine/Plugins/Enterprise/DatasmithFBXImporter/Source/DatasmithFBXTranslator/Public/DatasmithFBXScene.h
--- a/Engine/Plugins/Enterprise/DatasmithFBXImporter/Source/DatasmithFBXTranslator/Public/DatasmithFBXScene.h
+++ b/Engine/Plugins/Enterprise/DatasmithFBXImporter/Source/DatasmithFBXTranslator/Public/DatasmithFBXScene.h
@@ -135,6 +135,8 @@ struct DATASMITHFBXTRANSLATOR_API FDatasmithFBXSceneMaterial
 	TMap<FString, float> ScalarParams;
 	TMap<FString, bool> BoolParams;
 	TMap<FString, FTextureParams> TextureParams;
+
+	TArray<TWeakPtr<FDatasmithFBXSceneMaterial>> ClonedMaterials;
 };
 
 struct DATASMITHFBXTRANSLATOR_API FDatasmithFBXSceneMesh
