Remove unneeded include causing CIS error
#rnx

diff --git a/Engine/Plugins/Enterprise/DatasmithFBXImporter/Source/DatasmithDeltaGenTranslator/Private/DatasmithDeltaGenSceneProcessor.h b/Engine/Plugins/Enterprise/DatasmithFBXImporter/Source/DatasmithDeltaGenTranslator/Private/DatasmithDeltaGenSceneProcessor.h
--- a/Engine/Plugins/Enterprise/DatasmithFBXImporter/Source/DatasmithDeltaGenTranslator/Private/DatasmithDeltaGenSceneProcessor.h
+++ b/Engine/Plugins/Enterprise/DatasmithFBXImporter/Source/DatasmithDeltaGenTranslator/Private/DatasmithDeltaGenSceneProcessor.h
@@ -5,10 +5,10 @@
 #include "DatasmithFBXSceneProcessor.h"
 
 #include "CoreTypes.h"
-#include "Engine\EngineTypes.h"
 
 struct FDatasmithFBXScene;
 struct FDeltaGenTmlDataTimeline;
+struct FDirectoryPath;
 
 class FDatasmithDeltaGenSceneProcessor : public FDatasmithFBXSceneProcessor
 {
