VRED: Fix assert when parsing invalid transform file

diff --git a/Engine/Plugins/Enterprise/DatasmithFBXImporter/Source/DatasmithVREDTranslator/Private/DatasmithVREDImporterAuxFiles.cpp b/Engine/Plugins/Enterprise/DatasmithFBXImporter/Source/DatasmithVREDTranslator/Private/DatasmithVREDImporterAuxFiles.cpp
--- a/Engine/Plugins/Enterprise/DatasmithFBXImporter/Source/DatasmithVREDTranslator/Private/DatasmithVREDImporterAuxFiles.cpp
+++ b/Engine/Plugins/Enterprise/DatasmithFBXImporter/Source/DatasmithVREDTranslator/Private/DatasmithVREDImporterAuxFiles.cpp
@@ -69,9 +69,13 @@ void StringToFloatArray(const FString& InString, TArray<float>& OutArray)
 	}
 }
 
-void FloatArrayToTransform(const TArray<float>& Floats, FTransform& Transform)
+bool FloatArrayToTransform(const TArray<float>& Floats, FTransform& Transform)
 {
-	check(Floats.Num() == 16)
+	if ( Floats.Num() != 16 )
+	{
+		Transform.SetIdentity();
+		return false;
+	}
 
 	FMatrix matrix = FMatrix();
 	float (*matrixData)[4] = matrix.M;
@@ -84,6 +88,8 @@ void FloatArrayToTransform(const TArray<float>& Floats, FTransform& Transform)
 	FMemory::Memcpy(matrixData[3], &floatsData[12], 4 * sizeof(float));
 
 	Transform.SetFromMatrix(matrix);
+
+	return true;
 }
 
 void FixVREDXml(TArray<FString>& FileContentLines)
@@ -741,7 +747,11 @@ bool LoadVarFile(const TCHAR* InFilePath, FDatasmithVREDImportVariantsResult& Ou
 
 					TArray<float> TransformFloats;
 					StringToFloatArray(TransformData, TransformFloats);
-					FloatArrayToTransform(TransformFloats, TransVariant->Transform);
+					bool bSuccess = FloatArrayToTransform(TransformFloats, TransVariant->Transform);
+					if ( !bSuccess )
+					{
+						UE_LOG(LogDatasmithVREDImport, Warning, TEXT("Transform variant '%s' state '%s' has invalid transform value '%s'. Transform will be set to identity."), *TargetNodeName, *StateName, *TransformData);
+					}
 				}
 			}
 		}
