UCurveFloat::GetFloatValue returns a wrong value when passing very small negative values
#fix use double precision PI value for solving Cubic Equation

diff --git a/Engine/Plugins/Enterprise/DatasmithFBXImporter/Source/DatasmithDeltaGenTranslator/Private/DatasmithDeltaGenAnimationInterpolator.cpp b/Engine/Plugins/Enterprise/DatasmithFBXImporter/Source/DatasmithDeltaGenTranslator/Private/DatasmithDeltaGenAnimationInterpolator.cpp
--- a/Engine/Plugins/Enterprise/DatasmithFBXImporter/Source/DatasmithDeltaGenTranslator/Private/DatasmithDeltaGenAnimationInterpolator.cpp
+++ b/Engine/Plugins/Enterprise/DatasmithFBXImporter/Source/DatasmithDeltaGenTranslator/Private/DatasmithDeltaGenAnimationInterpolator.cpp
@@ -103,8 +103,8 @@ namespace DeltaGen
 			double T = 2 * sqrt(-P);
 
 			Solution[0] = T * cos(Phi);
-			Solution[1] = -T * cos(Phi + PI / 3);
-			Solution[2] = -T * cos(Phi - PI / 3);
+			Solution[1] = -T * cos(Phi + UE_DOUBLE_PI / 3);
+			Solution[2] = -T * cos(Phi - UE_DOUBLE_PI / 3);
 			NumSolutions = 3;
 		}
 		// One real solution
diff --git a/Engine/Source/Runtime/Engine/Private/Curves/RichCurve.cpp b/Engine/Source/Runtime/Engine/Private/Curves/RichCurve.cpp
--- a/Engine/Source/Runtime/Engine/Private/Curves/RichCurve.cpp
+++ b/Engine/Source/Runtime/Engine/Private/Curves/RichCurve.cpp
@@ -951,8 +951,8 @@ static int SolveCubic(double Coeff[4], double Solution[3])
 		double t = 2 * sqrt(-P);
 
 		Solution[0] = t * cos(phi);
-		Solution[1] = -t * cos(phi + PI / 3);
-		Solution[2] = -t * cos(phi - PI / 3);
+		Solution[1] = -t * cos(phi + UE_DOUBLE_PI / 3);
+		Solution[2] = -t * cos(phi - UE_DOUBLE_PI / 3);
 		NumSolutions = 3;
 	}
 	else /* one real solution */
diff --git a/Engine/Source/Runtime/MovieScene/Private/Channels/MovieSceneCurveChannelImpl.cpp b/Engine/Source/Runtime/MovieScene/Private/Channels/MovieSceneCurveChannelImpl.cpp
--- a/Engine/Source/Runtime/MovieScene/Private/Channels/MovieSceneCurveChannelImpl.cpp
+++ b/Engine/Source/Runtime/MovieScene/Private/Channels/MovieSceneCurveChannelImpl.cpp
@@ -198,8 +198,8 @@ static int SolveCubic(double Coeff[4], double Solution[3])
 		double t = 2 * sqrt(-P);
 
 		Solution[0] = t * cos(phi);
-		Solution[1] = -t * cos(phi + PI / 3);
-		Solution[2] = -t * cos(phi - PI / 3);
+		Solution[1] = -t * cos(phi + UE_DOUBLE_PI / 3);
+		Solution[2] = -t * cos(phi - UE_DOUBLE_PI / 3);
 		NumSolutions = 3;
 	}
 	else /* one real solution */
