Fixing CIS, PF somehow missed these files.

diff --git a/Engine/Plugins/Editor/SpeedTreeImporter/Source/SpeedTreeImporter/Private/SpeedTreeImportFactory.cpp b/Engine/Plugins/Editor/SpeedTreeImporter/Source/SpeedTreeImporter/Private/SpeedTreeImportFactory.cpp
--- a/Engine/Plugins/Editor/SpeedTreeImporter/Source/SpeedTreeImporter/Private/SpeedTreeImportFactory.cpp
+++ b/Engine/Plugins/Editor/SpeedTreeImporter/Source/SpeedTreeImporter/Private/SpeedTreeImportFactory.cpp
@@ -69,6 +69,7 @@
 #include "PackageTools.h"
 #include "SpeedTreeWind.h"
 #include "StaticMeshAttributes.h"
+#include "RenderMath.h"
 
 #if WITH_SPEEDTREE
 
diff --git a/Engine/Plugins/Enterprise/DatasmithFBXImporter/Source/DatasmithFBXTranslator/Private/DatasmithFBXFileImporter.cpp b/Engine/Plugins/Enterprise/DatasmithFBXImporter/Source/DatasmithFBXTranslator/Private/DatasmithFBXFileImporter.cpp
--- a/Engine/Plugins/Enterprise/DatasmithFBXImporter/Source/DatasmithFBXTranslator/Private/DatasmithFBXFileImporter.cpp
+++ b/Engine/Plugins/Enterprise/DatasmithFBXImporter/Source/DatasmithFBXTranslator/Private/DatasmithFBXFileImporter.cpp
@@ -14,6 +14,7 @@
 #include "StaticMeshAttributes.h"
 #include "Misc/Paths.h"
 #include "RenderUtils.h"
+#include "RenderMath.h"
 
 #define ROOT_NODE_NAME TEXT("RootNode")
 #define INCH_TO_MM(value) (value * 25.4)
diff --git a/Engine/Source/Developer/MeshSimplifier/Private/QuadricMeshReduction.cpp b/Engine/Source/Developer/MeshSimplifier/Private/QuadricMeshReduction.cpp
--- a/Engine/Source/Developer/MeshSimplifier/Private/QuadricMeshReduction.cpp
+++ b/Engine/Source/Developer/MeshSimplifier/Private/QuadricMeshReduction.cpp
@@ -15,6 +15,7 @@
 #include "RenderUtils.h"
 #include "Engine/StaticMesh.h"
 #include "Misc/ConfigCacheIni.h"
+#include "RenderMath.h"
 
 class FQuadricSimplifierMeshReductionModule : public IMeshReductionModule
 {
