Deactivate removal of degenerates/tiny polygons in datasmith FBX import. Requested by couple of companies importing from VRED exports.

diff --git a/Engine/Plugins/Enterprise/DatasmithFBXImporter/Source/DatasmithFBXTranslator/Private/DatasmithFBXFileImporter.cpp b/Engine/Plugins/Enterprise/DatasmithFBXImporter/Source/DatasmithFBXTranslator/Private/DatasmithFBXFileImporter.cpp
--- a/Engine/Plugins/Enterprise/DatasmithFBXImporter/Source/DatasmithFBXTranslator/Private/DatasmithFBXFileImporter.cpp
+++ b/Engine/Plugins/Enterprise/DatasmithFBXImporter/Source/DatasmithFBXTranslator/Private/DatasmithFBXFileImporter.cpp
@@ -1247,11 +1247,11 @@ void FDatasmithFBXFileImporter::DoImportMesh(FbxMesh* InMesh, FDatasmithFBXScene
 		}
 
 		// Skip degenerated polygons
-		FVector RawNormal = ((CornerPositions[1] - CornerPositions[2]) ^ (CornerPositions[0] - CornerPositions[2]));
+		/*FVector RawNormal = ((CornerPositions[1] - CornerPositions[2]) ^ (CornerPositions[0] - CornerPositions[2]));
 		if (RawNormal.SizeSquared() < SMALL_NUMBER)
 		{
 			continue; // this will leave holes...
-		}
+		}*/
 
 		// Create Vertex instances
 		CornerVertexInstanceIDs.SetNumUninitialized(CornerCount, false);
