Changing VRED detection from FBX header due to changes introduced by VRED 2024.

diff --git a/Engine/Plugins/Enterprise/DatasmithFBXImporter/Source/DatasmithVREDTranslator/Private/DatasmithVREDTranslator.cpp b/Engine/Plugins/Enterprise/DatasmithFBXImporter/Source/DatasmithVREDTranslator/Private/DatasmithVREDTranslator.cpp
--- a/Engine/Plugins/Enterprise/DatasmithFBXImporter/Source/DatasmithVREDTranslator/Private/DatasmithVREDTranslator.cpp
+++ b/Engine/Plugins/Enterprise/DatasmithFBXImporter/Source/DatasmithVREDTranslator/Private/DatasmithVREDTranslator.cpp
@@ -70,6 +70,23 @@ bool FDatasmithVREDTranslator::IsSourceSupported(const FDatasmithSceneSource& So
 			}
 		}
 	}
+	if(!bApplicationVendorValid)
+	{
+		const ANSICHAR* TagName = FPlatformString::Strstr(Header, "LastSaved|ApplicationVendor");
+		if (TagName)
+		{
+			const ANSICHAR* TagType = FPlatformString::Strstr(TagName, "KString");
+			if (TagType)
+			{
+				const ANSICHAR* TagData = FPlatformString::Strstr(TagType, "Autodesk");
+				if (TagData)
+				{
+					// The whole tag should be in the same vicinity
+					bApplicationVendorValid = (TagData - TagName) < 256;
+				}
+			}
+		}
+	}
 
 	bool bApplicationNameValid = false;
 	{
@@ -88,6 +105,23 @@ bool FDatasmithVREDTranslator::IsSourceSupported(const FDatasmithSceneSource& So
 			}
 		}
 	}
+	if(!bApplicationNameValid)
+	{
+		const ANSICHAR* TagName = FPlatformString::Strstr(Header, "LastSaved|ApplicationName");
+		if (TagName)
+		{
+			const ANSICHAR* TagType = FPlatformString::Strstr(TagName, "KString");
+			if (TagType)
+			{
+				const ANSICHAR* TagData = FPlatformString::Strstr(TagType, "VRED");
+				if (TagData)
+				{
+					// The whole tag should be in the same vicinity
+					bApplicationNameValid = (TagData - TagName) < 256;
+				}
+			}
+		}
+	}
 
 	return bApplicationVendorValid && bApplicationNameValid;
 }
